Moved StructUtils plugin out of NotForLicensees and marked it as experimental.

diff --git a/Engine/Plugins/Experimental/StructUtils/Config/BaseStructUtils.ini b/Engine/Plugins/Experimental/StructUtils/Config/BaseStructUtils.ini
new file mode 100644
--- /dev/null
+++ b/Engine/Plugins/Experimental/StructUtils/Config/BaseStructUtils.ini
@@ -0,0 +1,2 @@
+[CoreRedirects]
++StructRedirects=(OldName="UniqueScriptStructPtr", NewName="/Script/StructUtils.InstancedStruct")
diff --git a/Engine/Plugins/Experimental/StructUtils/Source/StructUtils/Private/InstancedStruct.cpp b/Engine/Plugins/Experimental/StructUtils/Source/StructUtils/Private/InstancedStruct.cpp
new file mode 100644
--- /dev/null
+++ b/Engine/Plugins/Experimental/StructUtils/Source/StructUtils/Private/InstancedStruct.cpp
@@ -0,0 +1,227 @@
+// Copyright Epic Games, Inc. All Rights Reserved.
+#include "InstancedStruct.h"
+
+void FInstancedStruct::InitializeAs(const UScriptStruct* InScriptStruct, const uint8* InStructMemory /*= nullptr*/)
+{
+	Reset();
+
+	if (!InScriptStruct)
+	{
+		// InScriptStruct == nullptr signifies an empty, unset FInstancedStruct instance. No further work required.
+		return;
+	}
+
+	SetScriptStruct(InScriptStruct);
+
+	const int32 RequiredSize = InScriptStruct->GetStructureSize();
+	SetMemory((uint8*)FMemory::Malloc(FMath::Max(1, RequiredSize)));
+
+	InScriptStruct->InitializeStruct(GetMutableMemory());
+
+	if (InStructMemory)
+	{
+		InScriptStruct->CopyScriptStruct(GetMutableMemory(), InStructMemory);
+	}
+}
+
+void FInstancedStruct::Reset()
+{
+	if (uint8* Memory = GetMutableMemory())
+	{
+		DestroyScriptStruct();
+		FMemory::Free(Memory);
+	}
+	FBaseStruct::Reset();
+}
+
+bool FInstancedStruct::Serialize(FArchive& Ar)
+{
+	UScriptStruct* NonConstStruct = const_cast<UScriptStruct*>(GetScriptStruct());
+
+	enum class EVersion : uint8
+	{
+		InitialVersion = 0,
+		// -----<new versions can be added above this line>-----
+		VersionPlusOne,
+		LatestVersion = VersionPlusOne - 1
+	};
+
+	EVersion Version = EVersion::LatestVersion;
+
+	// Temporary code to introduce versioning and load old data
+	// The goal is to remove this  by bumping the version in a near future.
+	bool bUseVersioning = true;
+
+#if WITH_EDITOR
+	if (!Ar.IsCooking())
+	{
+		// Keep archive position to use legacy serialization if the header is not found
+		const int64 HeaderOffset = Ar.Tell();
+
+		// Some random pattern to differentiate old data
+		const uint32 NewVersionHeader = 0xABABABAB;
+		uint32 Header = NewVersionHeader;
+		Ar << Header;
+
+		if (Ar.IsLoading())
+		{
+			if (Header != NewVersionHeader)
+			{
+				// Not a valid header so go back and process with legacy loading
+				Ar.Seek(HeaderOffset);
+				bUseVersioning = false;
+				UE_LOG(LogLoad, Verbose, TEXT("Loading FInstancedStruct using legacy serialization"));
+			}
+		}
+	}
+#endif // WITH_EDITOR
+
+	if (bUseVersioning)
+	{
+		Ar << Version;
+	}
+
+	if (Version > EVersion::LatestVersion)
+	{
+		UE_LOG(LogCore, Error, TEXT("Invalid Version: %hhu"), Version);
+		Ar.SetError();
+		return false;
+	}
+
+	if (Ar.IsLoading())
+	{
+		// UScriptStruct type
+		UScriptStruct* NewNonConstStruct = nullptr;
+		Ar << NewNonConstStruct;
+		NonConstStruct = ReinitializeAs(NewNonConstStruct);
+
+		// Size of the serialized memory
+		int32 SerialSize = 0; 
+		if (bUseVersioning)
+		{
+			Ar << SerialSize;
+		}
+
+		// Serialized memory
+		if (NonConstStruct == nullptr && SerialSize > 0)
+		{
+			// A null struct indicates an old struct or an unsupported one for the current target.
+			// In this case we manually seek in the archive to skip its serialized content. 
+			// We don't want to rely on TaggedSerialization that will mark an error in the archive that
+			// may cause other serialization to fail (e.g. FArchive& operator<<(FArchive& Ar, TArray& A))
+			UE_LOG(LogCore, Warning, TEXT("Unable to find serialized UScriptStruct -> Advance %u bytes in the archive and reset to empty FInstancedStruct"), SerialSize);
+			Ar.Seek(Ar.Tell() + SerialSize);
+		}
+		else if (NonConstStruct != nullptr && ensureMsgf(GetMutableMemory() != nullptr, TEXT("A valid script struct should always have allocated memory")))
+		{
+			NonConstStruct->SerializeItem(Ar, GetMutableMemory(), /* Defaults */ nullptr);
+		}
+	}
+	else if (Ar.IsSaving())
+	{
+		// UScriptStruct type
+		Ar << NonConstStruct;
+	
+		// Size of the serialized memory (reserve location)
+		const int64 SizeOffset = Ar.Tell(); // Position to write the actual size after struct serialization
+		int32 SerialSize = 0;
+		Ar << SerialSize;
+		
+		// Serialized memory
+		const int64 InitialOffset = Ar.Tell(); // Position before struct serialization to compute its serial size
+		if (NonConstStruct != nullptr && ensureMsgf(GetMutableMemory() != nullptr, TEXT("A valid script struct should always have allocated memory")))
+		{
+			NonConstStruct->SerializeItem(Ar, GetMutableMemory(), /* Defaults */ nullptr);
+		}
+		const int64 FinalOffset = Ar.Tell(); // Keep current offset to reset the archive pos after write the serial size
+
+		// Size of the serialized memory
+		Ar.Seek(SizeOffset);	// Go back in the archive to write the actual size
+		SerialSize = (int32)(FinalOffset - InitialOffset);
+		Ar << SerialSize;
+		Ar.Seek(FinalOffset);	// Reset archive to its position
+	}
+
+	return true;
+}
+
+bool FInstancedStruct::ExportTextItem(FString& ValueStr, FInstancedStruct const& DefaultValue, class UObject* Parent, int32 PortFlags, class UObject* ExportRootScope) const
+{
+	UScriptStruct* NonConstStruct = const_cast<UScriptStruct*>(GetScriptStruct());
+	if (NonConstStruct == nullptr || GetMemory() == nullptr)
+	{
+		return false;
+	}
+	
+	ValueStr += NonConstStruct->GetPathName();
+
+	NonConstStruct->ExportText(ValueStr, GetMemory(), GetMemory(), nullptr, EPropertyPortFlags::PPF_None, nullptr);
+	return true;
+}
+
+bool FInstancedStruct::ImportTextItem(const TCHAR*& Buffer, int32 PortFlags, UObject* Parent, FOutputDevice* ErrorText, FArchive* InSerializingArchive /*= nullptr*/)
+{
+	FString StructPath;
+	const TCHAR* NewBuffer = FPropertyHelpers::ReadToken(Buffer, /* out */ StructPath, /* dotted names */ true);
+	if (NewBuffer == nullptr)
+	{
+		return false;
+	}
+
+	Buffer = NewBuffer;
+		
+	UScriptStruct* NonConstStruct = ReinitializeAs(FindObject<UScriptStruct>(nullptr, *StructPath, false));
+	if (NonConstStruct == nullptr)
+	{
+		return false;
+	}
+
+	const TCHAR* Result = NonConstStruct->ImportText(Buffer, GetMutableMemory(), nullptr, EPropertyPortFlags::PPF_None, nullptr, NonConstStruct->GetName());
+	if (Result != nullptr)
+	{
+		Buffer = Result;
+	}
+
+	return true;
+}
+
+UScriptStruct* FInstancedStruct::ReinitializeAs(const UScriptStruct* InScriptStruct)
+{
+	UScriptStruct* NonConstStruct = const_cast<UScriptStruct*>(GetScriptStruct());
+
+	if (InScriptStruct != NonConstStruct)
+	{
+		InitializeAs(InScriptStruct, nullptr);
+		NonConstStruct = const_cast<UScriptStruct*>(GetScriptStruct());
+	}
+
+	return NonConstStruct;
+}
+
+bool FInstancedStruct::Identical(const FInstancedStruct* Other, uint32 PortFlags) const
+{
+	// Only empty is considered equal
+	return Other != nullptr && GetMemory() == nullptr && Other->GetMemory() == nullptr && GetScriptStruct() == nullptr && Other->GetScriptStruct() == nullptr;
+}
+
+void FInstancedStruct::AddStructReferencedObjects(class FReferenceCollector& Collector)
+{
+	if (const UScriptStruct* ScriptStructPtr = GetScriptStruct())
+	{
+		if (ScriptStructPtr->StructFlags & STRUCT_AddStructReferencedObjects)
+		{
+			ScriptStructPtr->GetCppStructOps()->AddStructReferencedObjects()(GetMutableMemory(), Collector);
+		}
+		else
+		{
+			// The iterator will recursively loop through all structs in structs too.
+			for (TPropertyValueIterator<const FObjectProperty> It(ScriptStructPtr, GetMemory()); It; ++It)
+			{
+				UObject** ObjectPtr = static_cast<UObject**>(const_cast<void*>(It.Value()));
+				Collector.AddReferencedObject(*ObjectPtr);
+			}
+		}
+	}
+}
+
+
diff --git a/Engine/Plugins/Experimental/StructUtils/Source/StructUtils/Private/StructUtilsModule.cpp b/Engine/Plugins/Experimental/StructUtils/Source/StructUtils/Private/StructUtilsModule.cpp
new file mode 100644
--- /dev/null
+++ b/Engine/Plugins/Experimental/StructUtils/Source/StructUtils/Private/StructUtilsModule.cpp
@@ -0,0 +1,32 @@
+// Copyright Epic Games, Inc. All Rights Reserved.
+
+#include "StructUtilsModule.h"
+#include "Modules/ModuleManager.h"
+
+#define LOCTEXT_NAMESPACE "StructUtils"
+
+class FStructUtilsModule : public IStructUtilsModule
+{
+	// Begin IModuleInterface
+	virtual void StartupModule() override;
+	virtual void ShutdownModule() override;
+};
+
+IMPLEMENT_MODULE(FStructUtilsModule, StructUtils)
+
+void FStructUtilsModule::StartupModule()
+{
+	// Called right after the module DLL has been loaded and the module object has been created
+	// Load dependent modules here, and they will be guaranteed to be available during ShutdownModule. ie:
+	// 		FModuleManager::Get().LoadModuleChecked(TEXT("HTTP"));
+}
+
+void FStructUtilsModule::ShutdownModule()
+{
+	// Called before the module is unloaded, right before the module object is destroyed.
+	// During normal shutdown, this is called in reverse order that modules finish StartupModule().
+	// This means that, as long as a module references dependent modules in it's StartupModule(), it
+	// can safely reference those dependencies in ShutdownModule() as well.
+}
+
+#undef LOCTEXT_NAMESPACE
diff --git a/Engine/Plugins/Experimental/StructUtils/Source/StructUtils/Public/InstancedStruct.h b/Engine/Plugins/Experimental/StructUtils/Source/StructUtils/Public/InstancedStruct.h
new file mode 100644
--- /dev/null
+++ b/Engine/Plugins/Experimental/StructUtils/Source/StructUtils/Public/InstancedStruct.h
@@ -0,0 +1,385 @@
+// Copyright Epic Games, Inc. All Rights Reserved.
+
+#pragma once
+
+#include "CoreMinimal.h"
+#include "UObject/Class.h"
+#include "UObject/WeakObjectPtr.h"
+#include "Templates/Casts.h"
+
+#include "InstancedStruct.generated.h"
+
+/**
+ * Base functionality for struct pointer handling, 
+ * Do not use directly, use either FInstancedStruct, FStructView
+ */
+USTRUCT()
+struct STRUCTUTILS_API FBaseStruct
+{
+	GENERATED_BODY()
+
+public:
+
+	FBaseStruct()
+	: StructMemory(nullptr)
+	{
+	}
+
+	// Copy and move should all be handled by subclasses, but it is impossible to delete them as it is a UStruct...
+	FBaseStruct(const FBaseStruct&) { checkf(false, TEXT("Should all be handled by subclasses")); }
+	FBaseStruct& operator=(const FBaseStruct&) { checkf(false, TEXT("Should all be handled by subclasses")); return *this; }
+	FBaseStruct(const FBaseStruct&&) { checkf(false, TEXT("Should all be handled by subclasses")); }
+	FBaseStruct& operator=(const FBaseStruct&&) { checkf(false, TEXT("Should all be handled by subclasses")); return *this; }
+
+	// Returns mutable reference to the struct, this getter assumes that all data is valid
+	template<typename T>
+	T& GetMutable()
+	{
+		check(StructMemory != nullptr);
+		check(ScriptStruct.Get());
+		check(ScriptStruct.Get()->IsChildOf(T::StaticStruct()));
+		return *((T*)StructMemory);
+	}
+
+	// Returns const reference to the struct, this getter assumes that all data is valid
+	template<typename T>
+	const T& Get() const
+	{
+		check(StructMemory != nullptr);
+		check(ScriptStruct.Get());
+		check(ScriptStruct.Get()->IsChildOf(T::StaticStruct()));
+		return *((T*)StructMemory);
+	}
+
+	// Returns mutable pointer to the struct, or nullptr if cast is not valid.
+	template<typename T>
+	T* GetMutablePtr()
+	{
+		if (StructMemory != nullptr && ScriptStruct.Get() && ScriptStruct.Get()->IsChildOf(T::StaticStruct()))
+		{
+			return ((T*)StructMemory);
+		}
+		return nullptr;
+	}
+
+	// Returns const pointer to the struct, or nullptr if cast is not valid.
+	template<typename T>
+	const T* GetPtr() const
+	{
+		if (StructMemory != nullptr && ScriptStruct.Get() && ScriptStruct.Get()->IsChildOf(T::StaticStruct()))
+		{
+			return ((T*)StructMemory);
+		}
+		return nullptr;
+	}
+
+	// Returns struct type.
+	const UScriptStruct* GetScriptStruct() const
+	{
+		return ScriptStruct.Get();
+	}
+
+	// Returns mutable pointer to struct memory.
+	uint8* GetMutableMemory()
+	{
+		return StructMemory;
+	}
+
+	// Returns const pointer to struct memory.
+	const uint8* GetMemory() const
+	{
+		return StructMemory;
+	}
+
+	/**
+	 * @return True if the struct is valid.
+	 */
+	bool IsValid() const
+	{
+		return StructMemory != nullptr && ScriptStruct.IsValid();
+	}
+
+	/**
+	 * Reset to empty.
+	 */
+	void Reset()
+	{
+		// We do not DestructStruct as we do not own it
+		StructMemory = nullptr;
+		ScriptStruct = nullptr;
+	}
+
+
+protected:
+
+	FBaseStruct(const UScriptStruct* InScriptStruct, uint8* InStructMemory = nullptr)
+		: ScriptStruct(InScriptStruct)
+		, StructMemory(InStructMemory)
+	{
+	}
+
+	void SetMemory(uint8* InStructMemory)
+	{
+		StructMemory = InStructMemory;
+	}
+
+	void SetScriptStruct(const UScriptStruct* InScriptStruct)
+	{
+		ScriptStruct = InScriptStruct;
+	}
+
+    void DestroyScriptStruct()
+	{
+		check(StructMemory != nullptr);
+		if (const UScriptStruct* ScriptStructPtr = ScriptStruct.Get())
+		{
+			ScriptStructPtr->DestroyStruct(StructMemory);
+		}
+	}
+
+private:
+	TWeakObjectPtr<const UScriptStruct> ScriptStruct;	// Struct type
+	uint8* StructMemory = nullptr;						// Struct memory
+};
+
+/**
+ * FInstancedStruct works similarly as instanced UObject* property but is USTRUCTs.
+ * Example:
+ *
+ *	UPROPERTY(EditAnywhere, Category = Foo, meta = (BaseStruct = "TestStructBase"))
+ *	FInstancedStruct Test;
+ *
+ *	UPROPERTY(EditAnywhere, Category = Foo, meta = (BaseStruct = "TestStructBase"))
+ *	TArray<FInstancedStruct> TestArray;
+ *
+ * -----------------------------------------------------------------------------
+ * DISCLAIMER: This will be removed when TUniquePtr<FScriptObject> is available!
+ * -----------------------------------------------------------------------------
+ */
+
+USTRUCT()
+struct STRUCTUTILS_API FInstancedStruct : public FBaseStruct
+{
+	GENERATED_BODY()
+
+public:
+
+	FInstancedStruct()
+		: FBaseStruct()
+	{
+	}
+
+	explicit FInstancedStruct(const UScriptStruct* InScriptStruct)
+		: FBaseStruct()
+	{
+		InitializeAs(InScriptStruct, nullptr);
+	}
+
+	FInstancedStruct(const FBaseStruct& InOther)
+		: FBaseStruct()
+	{
+		InitializeAs(InOther.GetScriptStruct(), InOther.GetMemory());
+	}
+
+	FInstancedStruct(const FInstancedStruct& InOther)
+		: FBaseStruct()
+	{
+		InitializeAs(InOther.GetScriptStruct(), InOther.GetMemory());
+	}
+
+	FInstancedStruct(FInstancedStruct&& InOther)
+		: FBaseStruct(InOther.GetScriptStruct(), InOther.GetMutableMemory())
+	{
+		InOther.SetScriptStruct(nullptr);
+		InOther.SetMemory(nullptr);
+	}
+
+
+	~FInstancedStruct()
+	{
+		Reset();
+	}
+
+	FInstancedStruct& operator=(const FBaseStruct& InOther)
+	{
+		if (this != &InOther)
+		{
+			InitializeAs(InOther.GetScriptStruct(), InOther.GetMemory());
+		}
+		return *this;
+	}
+
+	FInstancedStruct& operator=(const FInstancedStruct& InOther)
+	{
+		if (this != &InOther)
+		{
+			InitializeAs(InOther.GetScriptStruct(), InOther.GetMemory());
+		}
+		return *this;
+	}
+
+	FInstancedStruct& operator=(FInstancedStruct&& InOther)
+	{
+		if (this != &InOther)
+		{
+			Reset();
+
+			SetScriptStruct(InOther.GetScriptStruct());
+			SetMemory(InOther.GetMutableMemory());
+
+			InOther.SetScriptStruct(nullptr);
+			InOther.SetMemory(nullptr);
+		}
+		return *this;
+	}
+
+	// Initializes from struct type and optional data.
+	void InitializeAs(const UScriptStruct* InScriptStruct, const uint8* InStructMemory = nullptr);
+
+	// Initializes from struct type and emplace construct
+	template<typename T, typename... TArgs>
+	void InitializeAs(TArgs&&... InArgs)
+	{
+		Reset();
+
+		const UScriptStruct* InScriptStruct = T::StaticStruct();
+		SetScriptStruct(InScriptStruct);
+
+		const int32 RequiredSize = InScriptStruct->GetStructureSize();
+		SetMemory((uint8*)FMemory::Malloc(FMath::Max(1, RequiredSize)));
+
+		new (GetMutableMemory()) T(Forward<TArgs>(InArgs)...);
+	}
+
+	// Creates a new FInstancedStruct from templated struct type.
+	template<typename T>
+	static FInstancedStruct Make()
+	{
+		FInstancedStruct InstancedStruct;
+		InstancedStruct.InitializeAs(T::StaticStruct(), nullptr);
+		return InstancedStruct;
+	}
+
+	// Creates a new FInstancedStruct from templated struct.
+	template<typename T>
+	static FInstancedStruct Make(const T& Struct)
+	{
+		FInstancedStruct InstancedStruct;
+		InstancedStruct.InitializeAs(T::StaticStruct(), reinterpret_cast<const uint8*>(&Struct));
+		return InstancedStruct;
+	}
+
+	// Creates a new FInstancedStruct from the templated type and forward all arguments to constructor
+	template<typename T, typename... TArgs>
+	static inline FInstancedStruct Make(TArgs&&... InArgs)
+	{
+		FInstancedStruct InstancedStruct;
+		InstancedStruct.InitializeAs<T>(Forward<TArgs>(InArgs)...);
+		return InstancedStruct;
+	}
+
+	// Reset to empty.
+	void Reset();
+
+	// For StructOpsTypeTraits
+	bool Serialize(FArchive& Ar);
+	bool Identical(const FInstancedStruct* Other, uint32 PortFlags) const;
+	void AddStructReferencedObjects(class FReferenceCollector& Collector);
+	bool ExportTextItem(FString& ValueStr, FInstancedStruct const& DefaultValue, class UObject* Parent, int32 PortFlags, class UObject* ExportRootScope) const;
+	bool ImportTextItem(const TCHAR*& Buffer, int32 PortFlags, UObject* Parent, FOutputDevice* ErrorText, FArchive* InSerializingArchive = nullptr);
+
+private:
+	// Initializes for new struct type (does nothing if same type) and returns mutable struct
+	UScriptStruct* ReinitializeAs(const UScriptStruct* InScriptStruct);
+};
+
+template<>
+struct TStructOpsTypeTraits<FInstancedStruct> : public TStructOpsTypeTraitsBase2<FInstancedStruct>
+{
+	enum
+	{
+		WithSerializer = true,
+		WithIdentical = true,
+		WithExportTextItem = true,
+		WithImportTextItem = true,
+		WithAddStructReferencedObjects = true,
+	};
+};
+
+/** 
+ * FStructView is "typed" struct pointer, it contains pointer to struct plus UScriptStruct pointer.
+ * FStructView does not own the memory and will not free it when out of scope.
+ * It should be only used to pass struct pointer in a limited scope, or when the user controls the lifetime of the struct being stored.
+ * I.e. instead of passing ref or pointer to FInstancedStruct, you should use FStructView to pass around view to the contents.
+ * FStructView is passed by value.
+ * FStructView is similar to FStructOnScope, but FStructView is a view only (FStructOnScope can either own the memory or be a view)
+ *
+ * -----------------------------------------------------------------------------
+ * DISCLAIMER: This will be removed when FScriptObject is available!
+ * -----------------------------------------------------------------------------
+ */
+struct STRUCTUTILS_API FStructView : public FBaseStruct
+{
+
+public:
+
+	FStructView()
+		: FBaseStruct()
+	{
+	}
+
+	FStructView(const UScriptStruct* InScriptStruct, const uint8* InStructMemory)
+		: FBaseStruct(InScriptStruct, const_cast<uint8*>(InStructMemory))
+	{}
+
+	FStructView(const FBaseStruct& InOther)
+		: FStructView(InOther.GetScriptStruct(), InOther.GetMemory())
+	{
+	}
+
+	FStructView(const FStructView& InOther)
+		: FStructView(InOther.GetScriptStruct(), InOther.GetMemory())
+	{
+	}
+
+	FStructView(FStructView&& InOther)
+		: FStructView(InOther.GetScriptStruct(), InOther.GetMemory())
+	{
+	}
+
+	~FStructView()
+	{
+		Reset();
+	}
+
+	FStructView& operator=(const FBaseStruct& InOther)
+	{
+		SetScriptStruct(InOther.GetScriptStruct());
+		SetMemory(const_cast<uint8*>(InOther.GetMemory()));
+
+		return *this;
+	}
+
+	FStructView& operator=(const FStructView& InOther)
+	{
+		SetScriptStruct(InOther.GetScriptStruct());
+		SetMemory(const_cast<uint8*>(InOther.GetMemory()));
+
+		return *this;
+	}
+
+	FStructView& operator=(FStructView&& InOther)
+	{
+		SetScriptStruct(InOther.GetScriptStruct());
+		SetMemory(const_cast<uint8*>(InOther.GetMemory()));
+
+		return *this;
+	}
+
+	// Creates a new FStructView from the templated struct
+	template<typename T>
+	static FStructView Make(const T& Struct)
+	{
+		return FStructView(T::StaticStruct(), reinterpret_cast<const uint8*>(&Struct));
+	}
+};
\ No newline at end of file
diff --git a/Engine/Plugins/Experimental/StructUtils/Source/StructUtils/Public/StructUtilsModule.h b/Engine/Plugins/Experimental/StructUtils/Source/StructUtils/Public/StructUtilsModule.h
new file mode 100644
--- /dev/null
+++ b/Engine/Plugins/Experimental/StructUtils/Source/StructUtils/Public/StructUtilsModule.h
@@ -0,0 +1,37 @@
+// Copyright Epic Games, Inc. All Rights Reserved.
+
+#pragma  once
+
+#include "CoreMinimal.h"
+#include "Modules/ModuleManager.h"
+
+/**
+* The public interface to this module
+*/
+class IStructUtilsModule : public IModuleInterface
+{
+
+public:
+
+	/**
+	* Singleton-like access to this module's interface.  This is just for convenience!
+	* Beware of calling this during the shutdown phase, though.  Your module might have been unloaded already.
+	*
+	* @return Returns singleton instance, loading the module on demand if needed
+	*/
+	static inline IStructUtilsModule& Get()
+	{
+		return FModuleManager::LoadModuleChecked<IStructUtilsModule>("StructUtils");
+	}
+
+	/**
+	* Checks to see if this module is loaded and ready.  It is only valid to call Get() if IsAvailable() returns true.
+	*
+	* @return True if the module is loaded and ready to use
+	*/
+	static inline bool IsAvailable()
+	{
+		return FModuleManager::Get().IsModuleLoaded("StructUtils");
+	}
+};
+
diff --git a/Engine/Plugins/Experimental/StructUtils/Source/StructUtils/StructUtils.Build.cs b/Engine/Plugins/Experimental/StructUtils/Source/StructUtils/StructUtils.Build.cs
new file mode 100644
--- /dev/null
+++ b/Engine/Plugins/Experimental/StructUtils/Source/StructUtils/StructUtils.Build.cs
@@ -0,0 +1,32 @@
+// Copyright Epic Games, Inc. All Rights Reserved.
+
+namespace UnrealBuildTool.Rules
+{
+	public class StructUtils : ModuleRules
+	{
+		public StructUtils(ReadOnlyTargetRules Target) : base(Target)
+		{
+			PCHUsage = ModuleRules.PCHUsageMode.UseExplicitOrSharedPCHs;
+
+			PublicIncludePaths.AddRange(
+				new string[] {
+				ModuleDirectory + "/Public",
+				}
+			);
+
+			PublicDependencyModuleNames.AddRange(
+			new string[] {
+				"Core",
+				"CoreUObject",
+				"Engine",
+			}
+			);
+
+			PrivateDependencyModuleNames.AddRange(
+			new string[] {
+			}
+			);
+		}
+
+	}
+}
diff --git a/Engine/Plugins/Experimental/StructUtils/Source/StructUtilsEditor/Private/Customizations/InstancedStructDetails.cpp b/Engine/Plugins/Experimental/StructUtils/Source/StructUtilsEditor/Private/Customizations/InstancedStructDetails.cpp
new file mode 100644
--- /dev/null
+++ b/Engine/Plugins/Experimental/StructUtils/Source/StructUtilsEditor/Private/Customizations/InstancedStructDetails.cpp
@@ -0,0 +1,368 @@
+// Copyright Epic Games, Inc. All Rights Reserved.
+
+#include "InstancedStructDetails.h"
+#include "DetailWidgetRow.h"
+#include "DetailLayoutBuilder.h"
+#include "IPropertyUtilities.h"
+#include "IDetailPropertyRow.h"
+#include "IDetailChildrenBuilder.h"
+#include "IDetailCustomNodeBuilder.h"
+#include "Widgets/DeclarativeSyntaxSupport.h"
+#include "Widgets/Images/SImage.h"
+#include "Widgets/Input/SComboButton.h"
+#include "Framework/MultiBox/MultiBoxBuilder.h"
+#include "ScopedTransaction.h"
+#include "Modules/ModuleManager.h"
+#include "StructViewerModule.h"
+#include "StructViewerFilter.h"
+#include "Styling/SlateIconFinder.h"
+#include "Editor.h"
+#include "Engine/UserDefinedStruct.h"
+#include "InstancedStruct.h"
+
+#define LOCTEXT_NAMESPACE "StructUtilsEditor"
+
+class FInstancedStructFilter : public IStructViewerFilter
+{
+public:
+	/** The base struct for the property that classes must be a child-of. */
+	const UScriptStruct* BaseStruct = nullptr;
+
+	// A flag controlling whether we allow UserDefinedStructs
+	bool bAllowUserDefinedStructs = false;
+
+	// A flag controlling whether we allow to select the BaseStruct
+	bool bAllowBaseStruct = true;
+
+	virtual bool IsStructAllowed(const FStructViewerInitializationOptions& InInitOptions, const UScriptStruct* InStruct, TSharedRef<FStructViewerFilterFuncs> InFilterFuncs) override
+	{
+		if (InStruct->IsA<UUserDefinedStruct>())
+		{
+			return bAllowUserDefinedStructs;
+		}
+
+		if (InStruct == BaseStruct)
+		{
+			return bAllowBaseStruct;
+		}
+
+		// Query the native struct to see if it has the correct parent type (if any)
+		return !BaseStruct || InStruct->IsChildOf(BaseStruct);
+	}
+
+	virtual bool IsUnloadedStructAllowed(const FStructViewerInitializationOptions& InInitOptions, const FName InStructPath, TSharedRef<FStructViewerFilterFuncs> InFilterFuncs) override
+	{
+		// User Defined Structs don't support inheritance, so only include them requested
+		return bAllowUserDefinedStructs;
+	}
+};
+
+
+////////////////////////////////////
+
+FInstancedStructDataDetails::FInstancedStructDataDetails(TSharedPtr<IPropertyHandle> InStructProperty)
+{
+#if DO_CHECK
+	FStructProperty* StructProp = CastFieldChecked<FStructProperty>(InStructProperty->GetProperty());
+	check(StructProp);
+	check(StructProp->Struct == FInstancedStruct::StaticStruct());
+#endif
+
+	StructProperty = InStructProperty;
+}
+
+FInstancedStruct* FInstancedStructDataDetails::GetSingleInstancedStruct()
+{
+	if (!StructProperty)
+	{
+		return nullptr;
+	}
+	TArray<void*> RawData;
+	StructProperty->AccessRawData(RawData);
+	return RawData.Num() == 1 ? static_cast<FInstancedStruct*>(RawData[0]) : nullptr;
+}
+
+void FInstancedStructDataDetails::OnStructPreChange()
+{
+	if (StructProperty)
+	{
+		TArray<UObject*> OuterObjects;
+		StructProperty->GetOuterObjects(OuterObjects);
+		for (UObject* Outer : OuterObjects)
+		{
+			if (Outer)
+			{
+				Outer->Modify();
+			}
+		}
+	}
+}
+
+void FInstancedStructDataDetails::OnStructChanged()
+{
+	if (StructProperty)
+	{
+		StructProperty->NotifyFinishedChangingProperties();
+	}
+}
+
+void FInstancedStructDataDetails::SetOnRebuildChildren(FSimpleDelegate InOnRegenerateChildren)
+{
+	OnRegenerateChildren = InOnRegenerateChildren;
+}
+
+void FInstancedStructDataDetails::GenerateHeaderRowContent(FDetailWidgetRow& NodeRow)
+{
+}
+
+void FInstancedStructDataDetails::GenerateChildContent(IDetailChildrenBuilder& ChildBuilder)
+{
+	FInstancedStruct* InstancedStruct = GetSingleInstancedStruct();
+	if (!InstancedStruct)
+	{
+		return;
+	}
+
+	const UScriptStruct* ScriptStruct = InstancedStruct->GetScriptStruct();
+	uint8* Memory = InstancedStruct->GetMutableMemory();
+	if (!ScriptStruct || !Memory)
+	{
+		return;
+	}
+
+	// Note: this is potentially dangerous. It puts pointer to external memory on a shared ref.
+	TSharedRef<FStructOnScope> Struct = MakeShareable(new FStructOnScope(ScriptStruct, Memory));
+
+	TArray<UPackage*> OuterPackages;
+	StructProperty->GetOuterPackages(OuterPackages);
+	if (OuterPackages.Num() > 0)
+	{
+		Struct->SetPackage(OuterPackages[0]);
+	}
+
+	static const FName ShowOnlyInnerPropertiesMeta(TEXT("ShowOnlyInnerProperties"));
+	StructProperty->SetInstanceMetaData(ShowOnlyInnerPropertiesMeta, FString());
+
+	TArray<TSharedPtr<IPropertyHandle>> ChildProperties = StructProperty->AddChildStructure(Struct);
+	for (TSharedPtr<IPropertyHandle> ChildHandle : ChildProperties)
+	{
+		ChildHandle->SetOnPropertyValuePreChange(FSimpleDelegate::CreateSP(this, &FInstancedStructDataDetails::OnStructPreChange));
+		ChildHandle->SetOnChildPropertyValuePreChange(FSimpleDelegate::CreateSP(this, &FInstancedStructDataDetails::OnStructPreChange));
+		ChildHandle->SetOnPropertyValueChanged(FSimpleDelegate::CreateSP(this, &FInstancedStructDataDetails::OnStructChanged));
+		ChildHandle->SetOnChildPropertyValueChanged(FSimpleDelegate::CreateSP(this, &FInstancedStructDataDetails::OnStructChanged));
+
+		IDetailPropertyRow& Row = ChildBuilder.AddProperty(ChildHandle.ToSharedRef());
+		OnChildRowAdded(Row);
+	}
+}
+
+void FInstancedStructDataDetails::Tick(float DeltaTime)
+{
+	if (bRefresh)
+	{
+		OnRegenerateChildren.ExecuteIfBound();
+		bRefresh = false;
+	}
+}
+
+FName FInstancedStructDataDetails::GetName() const
+{
+	static const FName Name("InstancedStructDataDetails");
+	return Name;
+}
+
+void FInstancedStructDataDetails::PostUndo(bool bSuccess)
+{
+	// Regen here directly so that we do not access stale data from FStructOnScope created on GenerateChildContent().
+	OnRegenerateChildren.ExecuteIfBound();
+	bRefresh = false;
+}
+
+void FInstancedStructDataDetails::PostRedo(bool bSuccess)
+{
+	// Regen here directly so that we do not access stale data from FStructOnScope created on GenerateChildContent().
+	OnRegenerateChildren.ExecuteIfBound();
+	bRefresh = false;
+}
+
+////////////////////////////////////
+
+
+TSharedRef<IPropertyTypeCustomization> FInstancedStructDetails::MakeInstance()
+{
+	return MakeShareable(new FInstancedStructDetails);
+}
+
+void FInstancedStructDetails::CustomizeHeader(TSharedRef<class IPropertyHandle> StructPropertyHandle, class FDetailWidgetRow& HeaderRow, IPropertyTypeCustomizationUtils& StructCustomizationUtils)
+{
+	StructProperty = StructPropertyHandle;
+	PropUtils = StructCustomizationUtils.GetPropertyUtilities().Get();
+
+	// Find base class from meta data.
+	static const FName BaseClassMetaName(TEXT("BaseStruct"));
+	FString BaseClassName = StructProperty->GetMetaData(BaseClassMetaName);
+	BaseScriptStruct = FindObject<UScriptStruct>(ANY_PACKAGE, *BaseClassName);
+
+	static const FName StructTypeConst(TEXT("StructTypeConst"));
+	const bool bEnableStructSelection = !StructProperty->HasMetaData(StructTypeConst);
+
+	HeaderRow
+		.NameContent()
+		[
+			StructPropertyHandle->CreatePropertyNameWidget()
+		]
+		.ValueContent()
+		.MinDesiredWidth(250.f)
+		.VAlign(VAlign_Center)
+		[
+			SAssignNew(ComboButton, SComboButton)
+			.OnGetMenuContent(this, &FInstancedStructDetails::GenerateStructPicker)
+			.ContentPadding(0)
+			.IsEnabled(bEnableStructSelection)
+			.ButtonContent()
+			[
+				SNew(SHorizontalBox)
+				+ SHorizontalBox::Slot()
+				.AutoWidth()
+				.VAlign(VAlign_Center)
+				.Padding(0.0f, 0.0f, 4.0f, 0.0f)
+				[
+					SNew(SImage)
+					.Image(this, &FInstancedStructDetails::GetDisplayValueIcon)
+				]
+				+ SHorizontalBox::Slot()
+				.VAlign(VAlign_Center)
+				[
+					SNew(STextBlock)
+					.Text(this, &FInstancedStructDetails::GetDisplayValueString)
+					.Font(IDetailLayoutBuilder::GetDetailFont())
+				]
+			]
+		];
+}
+
+void FInstancedStructDetails::CustomizeChildren(TSharedRef<class IPropertyHandle> StructPropertyHandle, class IDetailChildrenBuilder& StructBuilder, IPropertyTypeCustomizationUtils& StructCustomizationUtils)
+{
+	TSharedRef<FInstancedStructDataDetails> DataDetails = MakeShareable(new FInstancedStructDataDetails(StructProperty));
+	StructBuilder.AddCustomBuilder(DataDetails);
+}
+
+const UScriptStruct* FInstancedStructDetails::GetCommonScriptStruct() const
+{
+	if (!StructProperty)
+	{
+		return nullptr;
+	}
+	const UScriptStruct* CommonScriptStruct = nullptr;
+	bool bMultipleValues = false;
+	TArray<void*> RawData;
+	StructProperty->AccessRawData(RawData);
+	for (void* Data : RawData)
+	{
+		if (FInstancedStruct* Struct = static_cast<FInstancedStruct*>(Data))
+		{
+			const UScriptStruct* ScriptStruct = Struct->GetScriptStruct();
+			if (!bMultipleValues && !CommonScriptStruct)
+			{
+				CommonScriptStruct = ScriptStruct;
+			}
+			else if (ScriptStruct != CommonScriptStruct)
+			{
+				CommonScriptStruct = nullptr;
+				bMultipleValues = true;
+			}
+		}
+	}
+
+	return CommonScriptStruct;
+}
+
+FText FInstancedStructDetails::GetDisplayValueString() const
+{
+	const UScriptStruct* ScriptStruct = GetCommonScriptStruct();
+	if (ScriptStruct)
+	{
+		return ScriptStruct->GetDisplayNameText();
+	}
+	return FText();
+}
+
+const FSlateBrush* FInstancedStructDetails::GetDisplayValueIcon() const
+{
+	return FSlateIconFinder::FindIconBrushForClass(UScriptStruct::StaticClass());
+}
+
+TSharedRef<SWidget> FInstancedStructDetails::GenerateStructPicker()
+{
+	static const FName ExcludeBaseStruct(TEXT("ExcludeBaseStruct"));
+	const bool bExcludeBaseStruct = StructProperty->HasMetaData(ExcludeBaseStruct);
+
+	TSharedRef<FInstancedStructFilter> StructFilter = MakeShared<FInstancedStructFilter>();
+	StructFilter->BaseStruct = BaseScriptStruct;
+	StructFilter->bAllowUserDefinedStructs = false;
+	StructFilter->bAllowBaseStruct = !bExcludeBaseStruct;
+
+	FStructViewerInitializationOptions Options;
+	Options.bShowUnloadedStructs = true;
+	Options.bShowNoneOption = true;
+	Options.StructFilter = StructFilter;
+	Options.NameTypeToDisplay = EStructViewerNameTypeToDisplay::DisplayName;
+	Options.DisplayMode = EStructViewerDisplayMode::ListView;
+	Options.bAllowViewOptions = true;
+
+	FOnStructPicked OnPicked(FOnStructPicked::CreateRaw(this, &FInstancedStructDetails::OnStructPicked));
+
+	return SNew(SBox)
+		.WidthOverride(280)
+		[
+			SNew(SVerticalBox)
+			+ SVerticalBox::Slot()
+			.AutoHeight()
+			.MaxHeight(500)
+			[
+				FModuleManager::LoadModuleChecked<FStructViewerModule>("StructViewer").CreateStructViewer(Options, OnPicked)
+			]
+		];
+}
+
+void FInstancedStructDetails::OnStructPicked(const UScriptStruct* InStruct)
+{
+	if (StructProperty)
+	{
+		GEditor->BeginTransaction(LOCTEXT("OnStructPicked", "Set Struct"));
+
+		StructProperty->NotifyPreChange();
+
+		TArray<void*> RawData;
+		StructProperty->AccessRawData(RawData);
+		for (void* Data : RawData)
+		{
+			if (FInstancedStruct* Struct = static_cast<FInstancedStruct*>(Data))
+			{
+				if (InStruct)
+				{
+					Struct->InitializeAs(InStruct);
+				}
+				else
+				{
+					Struct->Reset();
+				}
+			}
+		}
+
+		StructProperty->NotifyPostChange(EPropertyChangeType::ValueSet);
+		GEditor->EndTransaction();
+		StructProperty->NotifyFinishedChangingProperties();
+
+		GEditor->EndTransaction();
+	}
+
+	ComboButton->SetIsOpen(false);
+
+	if (PropUtils)
+	{
+		PropUtils->ForceRefresh();
+	}
+}
+
+#undef LOCTEXT_NAMESPACE
diff --git a/Engine/Plugins/Experimental/StructUtils/Source/StructUtilsEditor/Private/StructUtilsEditorModule.cpp b/Engine/Plugins/Experimental/StructUtils/Source/StructUtilsEditor/Private/StructUtilsEditorModule.cpp
new file mode 100644
--- /dev/null
+++ b/Engine/Plugins/Experimental/StructUtils/Source/StructUtilsEditor/Private/StructUtilsEditorModule.cpp
@@ -0,0 +1,30 @@
+// Copyright Epic Games, Inc. All Rights Reserved.
+
+#include "StructUtilsEditorModule.h"
+#include "Modules/ModuleManager.h"
+#include "InstancedStructDetails.h"
+
+#define LOCTEXT_NAMESPACE "StructUtilsEditor"
+
+IMPLEMENT_MODULE(FStructUtilsEditorModule, StructUtilsEditor)
+
+void FStructUtilsEditorModule::StartupModule()
+{
+	// Register the details customizer
+	FPropertyEditorModule& PropertyModule = FModuleManager::LoadModuleChecked<FPropertyEditorModule>("PropertyEditor");
+	PropertyModule.RegisterCustomPropertyTypeLayout("InstancedStruct", FOnGetPropertyTypeCustomizationInstance::CreateStatic(&FInstancedStructDetails::MakeInstance));
+	PropertyModule.NotifyCustomizationModuleChanged();
+}
+
+void FStructUtilsEditorModule::ShutdownModule()
+{
+	// Unregister the details customization
+	if (FModuleManager::Get().IsModuleLoaded("PropertyEditor"))
+	{
+		FPropertyEditorModule& PropertyModule = FModuleManager::LoadModuleChecked<FPropertyEditorModule>("PropertyEditor");
+		PropertyModule.UnregisterCustomPropertyTypeLayout("InstancedStruct");
+		PropertyModule.NotifyCustomizationModuleChanged();
+	}
+}
+
+#undef LOCTEXT_NAMESPACE
diff --git a/Engine/Plugins/Experimental/StructUtils/Source/StructUtilsEditor/Public/IStructUtilsEditor.h b/Engine/Plugins/Experimental/StructUtils/Source/StructUtilsEditor/Public/IStructUtilsEditor.h
new file mode 100644
--- /dev/null
+++ b/Engine/Plugins/Experimental/StructUtils/Source/StructUtilsEditor/Public/IStructUtilsEditor.h
@@ -0,0 +1,16 @@
+// Copyright Epic Games, Inc. All Rights Reserved.
+
+#pragma once
+
+#include "CoreMinimal.h"
+#include "Toolkits/AssetEditorToolkit.h"
+
+
+/** StructUtils Editor public interface */
+class STRUCTUTILSEDITOR_API IStructUtilsEditor : public FAssetEditorToolkit
+{
+public:
+
+};
+
+
diff --git a/Engine/Plugins/Experimental/StructUtils/Source/StructUtilsEditor/Public/InstancedStructDetails.h b/Engine/Plugins/Experimental/StructUtils/Source/StructUtilsEditor/Public/InstancedStructDetails.h
new file mode 100644
--- /dev/null
+++ b/Engine/Plugins/Experimental/StructUtils/Source/StructUtilsEditor/Public/InstancedStructDetails.h
@@ -0,0 +1,77 @@
+// Copyright Epic Games, Inc. All Rights Reserved.
+
+#pragma once
+
+#include "CoreMinimal.h"
+#include "UnrealClient.h"
+#include "IPropertyTypeCustomization.h"
+#include "IDetailCustomNodeBuilder.h"
+#include "EditorUndoClient.h"
+
+class IPropertyHandle;
+struct FInstancedStruct;
+class IDetailPropertyRow;
+
+/**
+ * Type customization for FInstancedStruct.
+ */
+class STRUCTUTILSEDITOR_API FInstancedStructDetails : public IPropertyTypeCustomization
+{
+public:
+	/** Makes a new instance of this detail layout class for a specific detail view requesting it */
+	static TSharedRef<IPropertyTypeCustomization> MakeInstance();
+
+	/** IPropertyTypeCustomization interface */
+	virtual void CustomizeHeader(TSharedRef<IPropertyHandle> StructPropertyHandle, class FDetailWidgetRow& HeaderRow, IPropertyTypeCustomizationUtils& StructCustomizationUtils) override;
+	virtual void CustomizeChildren(TSharedRef<IPropertyHandle> StructPropertyHandle, class IDetailChildrenBuilder& StructBuilder, IPropertyTypeCustomizationUtils& StructCustomizationUtils) override;
+
+private:
+	const UScriptStruct* GetCommonScriptStruct() const;
+	FText GetDisplayValueString() const;
+	const FSlateBrush* GetDisplayValueIcon() const;
+	TSharedRef<SWidget> GenerateStructPicker();
+	void OnStructPicked(const UScriptStruct* InStruct);
+
+	class IPropertyUtilities* PropUtils;
+	TSharedPtr<IPropertyHandle> StructProperty;
+	UScriptStruct* BaseScriptStruct = nullptr;
+	TSharedPtr<class SComboButton> ComboButton;
+};
+
+/** 
+ * Node builder for FInstancedStruct children.
+ * Expects property handle holding FInstancedStruct as input.
+ * Can be used in a implementation of a IPropertyTypeCustomization CustomizeChildren() to display editable FInstancedStruct contents.
+ * OnChildRowAdded() is called right after each property is added, which allows the property row to be customizable.
+ */
+class STRUCTUTILSEDITOR_API FInstancedStructDataDetails : public IDetailCustomNodeBuilder, public FSelfRegisteringEditorUndoClient, public TSharedFromThis<FInstancedStructDataDetails>
+{
+public:
+	FInstancedStructDataDetails(TSharedPtr<IPropertyHandle> InStructProperty);
+
+	//~ Begin IDetailCustomNodeBuilder interface
+	virtual void SetOnRebuildChildren(FSimpleDelegate InOnRegenerateChildren) override;
+	virtual void GenerateHeaderRowContent(FDetailWidgetRow& NodeRow) override;
+	virtual void GenerateChildContent(IDetailChildrenBuilder& ChildrenBuilder) override;
+	virtual void Tick(float DeltaTime) override;
+	virtual bool RequiresTick() const override { return true; }
+	virtual bool InitiallyCollapsed() const override { return false; }
+	virtual FName GetName() const override;
+	//~ End IDetailCustomNodeBuilder interface
+
+	/** FEditorUndoClient interface */
+	virtual void PostUndo(bool bSuccess) override;
+	virtual void PostRedo(bool bSuccess) override;
+
+	// Called when a child is added, override to customize a child row.
+	virtual void OnChildRowAdded(IDetailPropertyRow& ChildRow) {}
+
+private:
+	void OnStructPreChange();
+	void OnStructChanged();
+	FInstancedStruct* GetSingleInstancedStruct();
+
+	TSharedPtr<IPropertyHandle> StructProperty;
+	FSimpleDelegate OnRegenerateChildren;
+	bool bRefresh = false;
+};
diff --git a/Engine/Plugins/Experimental/StructUtils/Source/StructUtilsEditor/Public/StructUtilsEditorModule.h b/Engine/Plugins/Experimental/StructUtils/Source/StructUtilsEditor/Public/StructUtilsEditorModule.h
new file mode 100644
--- /dev/null
+++ b/Engine/Plugins/Experimental/StructUtils/Source/StructUtilsEditor/Public/StructUtilsEditorModule.h
@@ -0,0 +1,25 @@
+// Copyright Epic Games, Inc. All Rights Reserved.
+
+#pragma  once
+
+#include "CoreMinimal.h"
+#include "Modules/ModuleManager.h"
+#include "AssetTypeCategories.h"
+#include "Toolkits/IToolkitHost.h"
+#include "Toolkits/AssetEditorToolkit.h"
+
+
+class IStructUtilsEditor;
+class FAssetTypeActions_Base;
+struct FGraphPanelNodeFactory;
+
+/**
+* The public interface to this module
+*/
+class STRUCTUTILSEDITOR_API FStructUtilsEditorModule : public IModuleInterface//, public IHasMenuExtensibility, public IHasToolBarExtensibility
+{
+public:
+	// Begin IModuleInterface
+	virtual void StartupModule() override;
+	virtual void ShutdownModule() override;
+};
diff --git a/Engine/Plugins/Experimental/StructUtils/Source/StructUtilsEditor/StructUtilsEditor.Build.cs b/Engine/Plugins/Experimental/StructUtils/Source/StructUtilsEditor/StructUtilsEditor.Build.cs
new file mode 100644
--- /dev/null
+++ b/Engine/Plugins/Experimental/StructUtils/Source/StructUtilsEditor/StructUtilsEditor.Build.cs
@@ -0,0 +1,48 @@
+// Copyright Epic Games, Inc. All Rights Reserved.
+
+namespace UnrealBuildTool.Rules
+{
+	public class StructUtilsEditor : ModuleRules
+	{
+		public StructUtilsEditor(ReadOnlyTargetRules Target) : base(Target)
+		{
+			PCHUsage = ModuleRules.PCHUsageMode.UseExplicitOrSharedPCHs;
+
+			PublicIncludePaths.AddRange(
+			new string[] {
+			}
+			);
+
+			PublicDependencyModuleNames.AddRange(
+			new string[] {
+				"Core",
+				"CoreUObject",
+				"Engine",
+				"InputCore",
+				"AssetTools",
+				"UnrealEd",
+				"Slate",
+				"SlateCore",
+				"EditorStyle",
+				"PropertyEditor",
+				"AIModule",
+				"StructUtils",
+				"DetailCustomizations",
+				"ComponentVisualizers",
+			}
+			);
+
+			PrivateDependencyModuleNames.AddRange(
+			new string[] {
+				"RenderCore",
+				"GraphEditor",
+				"KismetWidgets",
+				"PropertyEditor",
+				"AIGraph",
+				"ToolMenus",
+			}
+			);
+		}
+
+	}
+}
diff --git a/Engine/Plugins/Experimental/StructUtils/StructUtils.uplugin b/Engine/Plugins/Experimental/StructUtils/StructUtils.uplugin
new file mode 100644
--- /dev/null
+++ b/Engine/Plugins/Experimental/StructUtils/StructUtils.uplugin
@@ -0,0 +1,35 @@
+{
+	"FileVersion": 1,
+	"Version": 1,
+	"VersionName": "0.1",
+	"FriendlyName": "Struct Utils",
+	"Description": "Experimental Struct Utilities supplying InstancedStruct type",
+	"Category": "Gameplay",
+	"CreatedBy": "Epic Games, Inc.",
+	"CreatedByURL": "http://epicgames.com",
+	"DocsURL": "",
+	"MarketplaceURL": "",
+	"SupportURL": "",
+	"EnabledByDefault": false,
+	"CanContainContent": false,
+	"IsBetaVersion": true,
+	"IsExperimentalVersion": true,
+	"Installed": false,
+	"Modules": [
+		{
+			"Name": "StructUtils",
+			"Type": "Runtime",
+			"LoadingPhase": "PreDefault"
+		},
+		{
+			"Name" : "StructUtilsEditor",
+			"Type" : "Editor"
+		}
+	],
+	"Plugins": [
+		{
+			"Name": "PropertyAccessEditor",
+			"Enabled": true
+		}
+	]
+}
\ No newline at end of file
