InstancedStruct/StructView: Removed weak object pointer from UScriptStruct.

diff --git a/Engine/Plugins/Experimental/StructUtils/Source/StructUtils/Private/InstancedStruct.cpp b/Engine/Plugins/Experimental/StructUtils/Source/StructUtils/Private/InstancedStruct.cpp
--- a/Engine/Plugins/Experimental/StructUtils/Source/StructUtils/Private/InstancedStruct.cpp
+++ b/Engine/Plugins/Experimental/StructUtils/Source/StructUtils/Private/InstancedStruct.cpp
@@ -208,6 +208,8 @@ void FInstancedStruct::AddStructReferencedObjects(class FReferenceCollector& Col
 {
 	if (const UScriptStruct* ScriptStructPtr = GetScriptStruct())
 	{
+		Collector.AddReferencedObject(ScriptStructPtr);
+
 		if (ScriptStructPtr->StructFlags & STRUCT_AddStructReferencedObjects)
 		{
 			ScriptStructPtr->GetCppStructOps()->AddStructReferencedObjects()(GetMutableMemory(), Collector);
diff --git a/Engine/Plugins/Experimental/StructUtils/Source/StructUtils/Public/InstancedStruct.h b/Engine/Plugins/Experimental/StructUtils/Source/StructUtils/Public/InstancedStruct.h
--- a/Engine/Plugins/Experimental/StructUtils/Source/StructUtils/Public/InstancedStruct.h
+++ b/Engine/Plugins/Experimental/StructUtils/Source/StructUtils/Public/InstancedStruct.h
@@ -36,8 +36,8 @@ public:
 	T& GetMutable()
 	{
 		check(StructMemory != nullptr);
-		check(ScriptStruct.Get());
-		check(ScriptStruct.Get()->IsChildOf(T::StaticStruct()));
+		check(ScriptStruct != nullptr);
+		check(ScriptStruct->IsChildOf(T::StaticStruct()));
 		return *((T*)StructMemory);
 	}
 
@@ -46,8 +46,8 @@ public:
 	const T& Get() const
 	{
 		check(StructMemory != nullptr);
-		check(ScriptStruct.Get());
-		check(ScriptStruct.Get()->IsChildOf(T::StaticStruct()));
+		check(ScriptStruct != nullptr);
+		check(ScriptStruct->IsChildOf(T::StaticStruct()));
 		return *((T*)StructMemory);
 	}
 
@@ -55,7 +55,7 @@ public:
 	template<typename T>
 	T* GetMutablePtr()
 	{
-		if (StructMemory != nullptr && ScriptStruct.Get() && ScriptStruct.Get()->IsChildOf(T::StaticStruct()))
+		if (StructMemory != nullptr && ScriptStruct && ScriptStruct->IsChildOf(T::StaticStruct()))
 		{
 			return ((T*)StructMemory);
 		}
@@ -66,7 +66,7 @@ public:
 	template<typename T>
 	const T* GetPtr() const
 	{
-		if (StructMemory != nullptr && ScriptStruct.Get() && ScriptStruct.Get()->IsChildOf(T::StaticStruct()))
+		if (StructMemory != nullptr && ScriptStruct && ScriptStruct->IsChildOf(T::StaticStruct()))
 		{
 			return ((T*)StructMemory);
 		}
@@ -76,7 +76,7 @@ public:
 	// Returns struct type.
 	const UScriptStruct* GetScriptStruct() const
 	{
-		return ScriptStruct.Get();
+		return ScriptStruct;
 	}
 
 	// Returns mutable pointer to struct memory.
@@ -96,7 +96,7 @@ public:
 	 */
 	bool IsValid() const
 	{
-		return StructMemory != nullptr && ScriptStruct.IsValid();
+		return StructMemory != nullptr && ScriptStruct != nullptr;
 	}
 
 	/**
@@ -131,15 +131,15 @@ protected:
     void DestroyScriptStruct()
 	{
 		check(StructMemory != nullptr);
-		if (const UScriptStruct* ScriptStructPtr = ScriptStruct.Get())
+		if (ScriptStruct != nullptr)
 		{
-			ScriptStructPtr->DestroyStruct(StructMemory);
+			ScriptStruct->DestroyStruct(StructMemory);
 		}
 	}
 
 private:
-	TWeakObjectPtr<const UScriptStruct> ScriptStruct;	// Struct type
-	uint8* StructMemory = nullptr;						// Struct memory
+	const UScriptStruct* ScriptStruct = nullptr;	// Struct type
+	uint8* StructMemory = nullptr;					// Struct memory
 };
 
 /**
