InstancedStructStream
- Append only bufferfor heterogenous structs

diff --git a/Engine/Plugins/Experimental/StructUtils/Source/StructUtils/Private/ChunkedStructBuffer.cpp b/Engine/Plugins/Experimental/StructUtils/Source/StructUtils/Private/ChunkedStructBuffer.cpp
new file mode 100644
--- /dev/null
+++ b/Engine/Plugins/Experimental/StructUtils/Source/StructUtils/Private/ChunkedStructBuffer.cpp
@@ -0,0 +1,133 @@
+﻿// Copyright Epic Games, Inc. All Rights Reserved.
+#include "ChunkedStructBuffer.h"
+
+void FChunkedStructBuffer::Reset()
+{
+	ForEach([](FStructView View)
+	{
+		View.GetScriptStruct()->DestroyStruct(View.GetMutableMemory());
+	});
+
+	// Cleanup chunks
+	for (FChunkHeader* Chunk = First; Chunk != nullptr; Chunk = Chunk->Next)
+	{
+		for (int32 Index = 0; Index < MaxScriptStructsPerChunk; Index++)
+		{
+			Chunk->ScriptStructs[Index] = nullptr;
+		}
+		Chunk->UsedSize = 0;
+		Chunk->NumScriptStructs = 0;
+		Chunk->NumItems = 0;
+	}
+	NumItems = 0;
+
+	// Return chunks to freelist.
+	if (Head != nullptr)
+	{
+		check(First != nullptr);
+		Head->Next = FreeList;
+		FreeList = First;
+	}
+	
+	First = nullptr;
+	Head = nullptr;
+}
+
+void FChunkedStructBuffer::Compact()
+{
+	// Free FreeList chunks
+	FChunkHeader* Chunk = FreeList;
+	while (Chunk)
+	{
+		FChunkHeader* NextChunk = Chunk->Next;
+		FMemory::Free(Chunk);
+		Chunk = NextChunk;
+	}
+	FreeList = nullptr;
+}
+
+void FChunkedStructBuffer::Clear()
+{
+	Reset();
+	Compact();
+}
+
+void FChunkedStructBuffer::AllocateNewChunk()
+{
+	check(ChunkSize > 0);
+
+	// Allocate new chunk, use existing chunk from free list if possible.
+	FChunkHeader* Chunk = nullptr;
+	if (FreeList)
+	{
+		Chunk = FreeList;
+		FreeList = FreeList->Next;
+	}
+	else
+	{
+		Chunk = (FChunkHeader*)FMemory::Malloc(sizeof(FChunkHeader) + ChunkSize);
+	}
+	check(Chunk != nullptr);
+
+	Chunk->UsedSize = 0;
+	Chunk->NumScriptStructs = 0;
+	Chunk->NumItems = 0;
+	Chunk->Next = nullptr;
+
+	// Insert chunk the linked list.
+	if (First == nullptr)
+	{
+		First = Chunk;
+	}
+	
+	if (Head != nullptr)
+	{
+		Head->Next = Chunk;
+	}
+	
+	Head = Chunk;
+
+	check(First != nullptr);
+	check(Head != nullptr);
+}
+
+void FChunkedStructBuffer::AddStructReferencedObjects(class FReferenceCollector& Collector)
+{
+	for (FChunkHeader* Chunk = First; Chunk != nullptr; Chunk = Chunk->Next)
+	{
+		// Add reference to the ScriptStruct object.
+		for (int32 Index = 0; Index < Chunk->NumScriptStructs; Index++)
+		{
+			Collector.AddReferencedObject(Chunk->ScriptStructs[Index]);
+		}
+
+		// Add references in the struct contents.
+		uint8* ItemMemory = reinterpret_cast<uint8*>(Chunk) + sizeof(FChunkHeader);
+		for (int32 Index = 0; Index < Chunk->NumItems; Index++)
+		{
+			FItemHeader& Item = *reinterpret_cast<FItemHeader*>(ItemMemory);
+			const UScriptStruct* ScriptStruct = Chunk->ScriptStructs[Item.ScriptStructIndex];
+			uint8* StructMemory = ItemMemory + Item.Offset;
+			ItemMemory += Item.Size;
+			
+			if (ScriptStruct->StructFlags & STRUCT_AddStructReferencedObjects)
+			{
+				ScriptStruct->GetCppStructOps()->AddStructReferencedObjects()(StructMemory, Collector);
+			}
+			else if (ScriptStruct->StructFlags & STRUCT_IsPlainOldData)
+			{
+				// Dont bother with POD types.
+				continue;
+			}
+			else
+			{
+				// The iterator will recursively loop through all structs in structs too.
+				for (TPropertyValueIterator<const FObjectProperty> It(ScriptStruct, StructMemory); It; ++It)
+				{
+					UObject** ObjectPtr = static_cast<UObject**>(const_cast<void*>(It.Value()));
+					Collector.AddReferencedObject(*ObjectPtr);
+				}
+			}
+		}
+	}
+}
\ No newline at end of file
diff --git a/Engine/Plugins/Experimental/StructUtils/Source/StructUtils/Private/InstancedStructStream.cpp b/Engine/Plugins/Experimental/StructUtils/Source/StructUtils/Private/InstancedStructStream.cpp
new file mode 100644
--- /dev/null
+++ b/Engine/Plugins/Experimental/StructUtils/Source/StructUtils/Private/InstancedStructStream.cpp
@@ -0,0 +1,133 @@
+﻿// Copyright Epic Games, Inc. All Rights Reserved.
+#include "InstancedStructStream.h"
+
+void FInstancedStructStream::Reset()
+{
+	ForEach([](FStructView View)
+	{
+		View.GetScriptStruct()->DestroyStruct(View.GetMutableMemory());
+	});
+
+	// Cleanup chunks
+	for (FChunkHeader* Chunk = First; Chunk != nullptr; Chunk = Chunk->Next)
+	{
+		for (int32 Index = 0; Index < MaxScriptStructsPerChunk; Index++)
+		{
+			Chunk->ScriptStructs[Index] = nullptr;
+		}
+		Chunk->UsedSize = 0;
+		Chunk->NumScriptStructs = 0;
+		Chunk->NumItems = 0;
+	}
+	NumItems = 0;
+
+	// Return chunks to freelist.
+	if (Head != nullptr)
+	{
+		check(First != nullptr);
+		Head->Next = FreeList;
+		FreeList = First;
+	}
+	
+	First = nullptr;
+	Head = nullptr;
+}
+
+void FInstancedStructStream::Compact()
+{
+	// Free FreeList chunks
+	FChunkHeader* Chunk = FreeList;
+	while (Chunk)
+	{
+		FChunkHeader* NextChunk = Chunk->Next;
+		FMemory::Free(Chunk);
+		Chunk = NextChunk;
+	}
+	FreeList = nullptr;
+}
+
+void FInstancedStructStream::Clear()
+{
+	Reset();
+	Compact();
+}
+
+void FInstancedStructStream::AllocateNewChunk()
+{
+	check(ChunkSize > 0);
+
+	// Allocate new chunk, use existing chunk from free list if possible.
+	FChunkHeader* Chunk = nullptr;
+	if (FreeList)
+	{
+		Chunk = FreeList;
+		FreeList = FreeList->Next;
+	}
+	else
+	{
+		Chunk = (FChunkHeader*)FMemory::Malloc(sizeof(FChunkHeader) + ChunkSize);
+	}
+	check(Chunk != nullptr);
+
+	Chunk->UsedSize = 0;
+	Chunk->NumScriptStructs = 0;
+	Chunk->NumItems = 0;
+	Chunk->Next = nullptr;
+
+	// Insert chunk the linked list.
+	if (First == nullptr)
+	{
+		First = Chunk;
+	}
+	
+	if (Head != nullptr)
+	{
+		Head->Next = Chunk;
+	}
+	
+	Head = Chunk;
+
+	check(First != nullptr);
+	check(Head != nullptr);
+}
+
+void FInstancedStructStream::AddStructReferencedObjects(class FReferenceCollector& Collector)
+{
+	for (FChunkHeader* Chunk = First; Chunk != nullptr; Chunk = Chunk->Next)
+	{
+		// Add reference to the ScriptStruct object.
+		for (int32 Index = 0; Index < Chunk->NumScriptStructs; Index++)
+		{
+			Collector.AddReferencedObject(Chunk->ScriptStructs[Index]);
+		}
+
+		// Add references in the struct contents.
+		uint8* ItemMemory = reinterpret_cast<uint8*>(Chunk) + sizeof(FChunkHeader);
+		for (int32 Index = 0; Index < Chunk->NumItems; Index++)
+		{
+			FItemHeader& Item = *reinterpret_cast<FItemHeader*>(ItemMemory);
+			const UScriptStruct* ScriptStruct = Chunk->ScriptStructs[Item.ScriptStructIndex];
+			uint8* StructMemory = ItemMemory + Item.Offset;
+			ItemMemory += Item.Size;
+			
+			if (ScriptStruct->StructFlags & STRUCT_AddStructReferencedObjects)
+			{
+				ScriptStruct->GetCppStructOps()->AddStructReferencedObjects()(StructMemory, Collector);
+			}
+			else if (ScriptStruct->StructFlags & STRUCT_IsPlainOldData)
+			{
+				// Dont bother with POD types.
+				continue;
+			}
+			else
+			{
+				// The iterator will recursively loop through all structs in structs too.
+				for (TPropertyValueIterator<const FObjectProperty> It(ScriptStruct, StructMemory); It; ++It)
+				{
+					UObject** ObjectPtr = static_cast<UObject**>(const_cast<void*>(It.Value()));
+					Collector.AddReferencedObject(*ObjectPtr);
+				}
+			}
+		}
+	}
+}
\ No newline at end of file
diff --git a/Engine/Plugins/Experimental/StructUtils/Source/StructUtils/Public/ChunkedStructBuffer.h b/Engine/Plugins/Experimental/StructUtils/Source/StructUtils/Public/ChunkedStructBuffer.h
new file mode 100644
--- /dev/null
+++ b/Engine/Plugins/Experimental/StructUtils/Source/StructUtils/Public/ChunkedStructBuffer.h
@@ -0,0 +1,389 @@
+﻿// Copyright Epic Games, Inc. All Rights Reserved.
+
+#pragma once
+
+#include "CoreMinimal.h"
+#include "UObject/Class.h"
+#include "InstancedStruct.h"
+#include "ChunkedStructBuffer.generated.h"
+
+/**
+ * Buffer where you can append heterogeneous structs, and iterate over them in order (no random access).
+ * Note: If you use move assign or Append(&&), make sure to call Compact() or you may end up with large free list.
+ */
+USTRUCT()
+struct STRUCTUTILS_API FChunkedStructBuffer
+{
+	GENERATED_BODY()
+
+	static constexpr int32 DefaultChunkSize = 16384;
+
+	FChunkedStructBuffer(const int32 InChunkSize = DefaultChunkSize)
+		: ChunkSize(InChunkSize)
+	{
+	}
+
+	FChunkedStructBuffer(const FChunkedStructBuffer& Other)
+	{
+		ChunkSize = Other.ChunkSize;
+		Append(Other);
+	}
+	
+	FChunkedStructBuffer(FChunkedStructBuffer&& Other)
+	{
+		ChunkSize = Other.ChunkSize;
+		Append(Other);
+	}
+
+	~FChunkedStructBuffer()
+	{
+		Clear();
+	}
+
+	FChunkedStructBuffer& operator=(const FChunkedStructBuffer& Other)
+	{
+		Reset();
+		Append(Other);
+		return *this;
+	}
+	
+	FChunkedStructBuffer& operator=(FChunkedStructBuffer&& Other)
+	{
+		Reset();
+		Append(Other);
+		return *this;
+	}
+
+	void Append(const FChunkedStructBuffer& Other)
+	{
+		Other.ForEach([this](FStructView View)
+		{
+			Add(View);
+		});
+	}
+	
+	void Append(FChunkedStructBuffer&& Other)
+	{
+		if (Other.ChunkSize == ChunkSize)
+		{
+			// If same chunk size, steal chunks.
+			NumItems += Other.NumItems;
+			Head->Next = Other.First;
+			Head = Other.Head;
+
+			Other.NumItems = 0;
+			Other.First = nullptr;
+			Other.Head = nullptr;
+		}
+		else
+		{
+			Other.ForEach([this](FStructView View)
+			{
+				Add(View);
+			});
+		}
+		Other.Clear();
+	}
+
+	/** Emplaces struct in the buffer. */
+	template<typename T, typename... TArgs>
+	void Emplace(TArgs&&... InArgs)
+	{
+		const UScriptStruct* InScriptStruct = T::StaticStruct();
+		uint8* ItemMemory = AllocItem(InScriptStruct);
+		new (ItemMemory) T(Forward<TArgs>(InArgs)...);
+	}
+
+	/** Emplaces struct in the buffer, and returns reference to it. */
+	template<typename T, typename... TArgs>
+	T& Emplace_GetRef(TArgs&&... InArgs)
+	{
+		const UScriptStruct* InScriptStruct = T::StaticStruct();
+		uint8* ItemMemory = AllocItem(InScriptStruct);
+		return *new (ItemMemory) T(Forward<TArgs>(InArgs)...);
+	}
+
+	/** Add new struct in the buffer. */
+	template<typename T>
+	void Add(const T& InStruct)
+	{
+		const UScriptStruct* InScriptStruct = T::StaticStruct();
+		uint8* ItemMemory = AllocItem(InScriptStruct);
+		new (ItemMemory) T(InStruct);
+	}
+
+	/** Add new struct in the buffer, and return reference to it. */
+	template<typename T>
+	T& Add_GetRef(const T& InStruct)
+	{
+		const UScriptStruct* InScriptStruct = T::StaticStruct();
+		uint8* ItemMemory = AllocItem(InScriptStruct);
+		return *new (ItemMemory) T(InStruct);
+	}
+
+	/** Adds struct to the buffer based on ScriptStruct and pointer. */
+	void Add(const FStructView Struct)
+	{
+		check(Struct.IsValid());
+		const UScriptStruct* InScriptStruct = Struct.GetScriptStruct();
+		uint8* ItemMemory = AllocItem(InScriptStruct);
+		InScriptStruct->InitializeStruct(ItemMemory);
+		InScriptStruct->CopyScriptStruct(ItemMemory, Struct.GetMemory());
+	}
+
+	/** Iterates over all structs and calls Function on each item */
+	template<typename TFunc>
+	void ForEach(TFunc&& Function) const
+	{
+		for (FChunkHeader* Chunk = First; Chunk != nullptr; Chunk = Chunk->Next)
+		{
+			uint8* ItemMemory = reinterpret_cast<uint8*>(Chunk) + sizeof(FChunkHeader);
+			for (int32 Index = 0; Index < Chunk->NumItems; Index++)
+			{
+				FItemHeader& Item = *reinterpret_cast<FItemHeader*>(ItemMemory);
+				check(Item.ScriptStructIndex < MaxScriptStructsPerChunk);
+			
+				Function(FStructView(Chunk->ScriptStructs[Item.ScriptStructIndex], ItemMemory + Item.Offset));
+
+				ItemMemory += Item.Size;
+			}
+		}
+	}
+
+	/**
+	 * Iterates over all structs of specified type T and calls Function on each item.
+	 * Usage: Buffer.ForEach<FFoo>([](Foo& Item) { ... });
+	 */
+	template<typename T, typename TFunc>
+	void ForEach(TFunc&& Function) const
+	{
+		for (FChunkHeader* Chunk = First; Chunk != nullptr; Chunk = Chunk->Next)
+		{
+			// Build mask of the accepted struct in the chunk.
+			const int32 StructIndex = Chunk->GetScriptStructIndex(T::StaticStruct());
+			if (StructIndex == INDEX_NONE)
+			{
+				continue;
+			}
+			const uint32 FilterMask = 1U << StructIndex;
+			
+			uint8* ItemMemory = reinterpret_cast<uint8*>(Chunk) + sizeof(FChunkHeader);
+			for (int32 Index = 0; Index < Chunk->NumItems; Index++)
+			{
+				FItemHeader& Item = *reinterpret_cast<FItemHeader*>(ItemMemory);
+				const uint32 StructMask = 1U << Item.ScriptStructIndex;
+
+				if ((FilterMask & StructMask) != 0)
+				{
+					check(Item.ScriptStructIndex < MaxScriptStructsPerChunk);
+					Function(*reinterpret_cast<T*>(ItemMemory + Item.Offset));
+				}
+
+				ItemMemory += Item.Size;
+			}
+		}
+	}
+
+	/** Iterates over all structs of specified types and calls Function on each item */
+	template<typename TFunc>
+	void ForEachFiltered(TArrayView<const UScriptStruct*> AcceptedScriptStructs, TFunc&& Function) const
+	{
+		for (FChunkHeader* Chunk = First; Chunk != nullptr; Chunk = Chunk->Next)
+		{
+			// Build mask of all accepted structs in the chunk.
+			uint32 FilterMask = 0;
+			for (const UScriptStruct* ScriptStruct : AcceptedScriptStructs)
+			{
+				const int32 StructIndex = Chunk->GetScriptStructIndex(ScriptStruct);
+				if (StructIndex != INDEX_NONE)
+				{
+					FilterMask |= 1U << StructIndex;
+				}
+			}
+
+			// Chunk does not match any of the structs specified in the filter.
+			if (FilterMask == 0)
+			{
+				continue;
+			}
+			
+			uint8* ItemMemory = reinterpret_cast<uint8*>(Chunk) + sizeof(FChunkHeader);
+			for (int32 Index = 0; Index < Chunk->NumItems; Index++)
+			{
+				FItemHeader& Item = *reinterpret_cast<FItemHeader*>(ItemMemory);
+				const uint32 StructMask = 1U << Item.ScriptStructIndex;
+
+				if ((FilterMask & StructMask) != 0)
+				{
+					check(Item.ScriptStructIndex < MaxScriptStructsPerChunk);
+					Function(FStructView(Chunk->ScriptStructs[Item.ScriptStructIndex], ItemMemory + Item.Offset));
+				}
+
+				ItemMemory += Item.Size;
+			}
+		}
+	}
+
+	/** Returns struct types in the buffer */
+	void GetScriptStructs(TArray<const UScriptStruct*>& OutScriptStructs) const
+	{
+		OutScriptStructs.Reset();
+		for (const FChunkHeader* Chunk = First; Chunk != nullptr; Chunk = Chunk->Next)
+		{
+			// Add reference to the ScriptStruct object.
+			for (int32 Index = 0; Index < Chunk->NumScriptStructs; Index++)
+			{
+				OutScriptStructs.AddUnique(Chunk->ScriptStructs[Index]);
+			}
+		}
+	}
+
+	/** Resets and clears all structs, keeps internal memory. */
+	void Reset();
+
+	/** Releases unused internal memory. */
+	void Compact();
+
+	/** Resets and clears all structs, and releases unused internal memory. */
+	void Clear();
+
+	/** @return Number of structs added to the buffer. */
+	int32 Num() const { return NumItems; }
+
+	/** @return True if the buffer is empty. */
+	bool IsEmpty() const { return NumItems == 0; }
+
+	/** @return Chunk size */
+	int32 GetChunkSize() const { return ChunkSize; }
+
+	/** @return Number of chunks in use. */
+	int32 GetNumUsedChunks() const
+	{
+		int32 Count = 0;
+		for (const FChunkHeader* Chunk = First; Chunk != nullptr; Chunk = Chunk->Next)
+		{
+			Count++;
+		}
+		return Count;
+	}
+
+	/** @return Number of  chunks in freelist. */
+	int32 GetNumFreeChunks() const
+	{
+		int32 Count = 0;
+		for (const FChunkHeader* Chunk = FreeList; Chunk != nullptr; Chunk = Chunk->Next)
+		{
+			Count++;
+		}
+		return Count;
+	}
+
+	void AddStructReferencedObjects(class FReferenceCollector& Collector);
+
+protected:
+
+	/** Number of unique struct types that can be stored per chunk. */
+	static constexpr int32 MaxScriptStructsPerChunk = 8;
+
+	/** Header for each item in a chunk. Allocated before the item in tge chunk memory. */
+	struct FItemHeader
+	{
+		int32 Size;					/** Size of the item including this header */
+		int16 Offset;				/** relative offset to the item data */ 
+		int16 ScriptStructIndex;	/** Index in chunk ScriptStruct array to the type of the item */ 
+	};
+
+	/** Header for a chunk of items. */
+	struct FChunkHeader
+	{
+		/** @returns Index to the ScriptStruct stored in the chunk, or INDEX_NONE if cannot store the ScriptStruct. */
+		int32 GetScriptStructIndex(const UScriptStruct* InScriptStruct)
+		{
+			for (int32 Index = 0; Index < NumScriptStructs; Index++)
+			{
+				if (ScriptStructs[Index] == InScriptStruct)
+				{
+					return Index;
+				}
+			}
+
+			if (NumScriptStructs < MaxScriptStructsPerChunk)
+			{
+				ScriptStructs[NumScriptStructs++] = InScriptStruct;
+				return NumScriptStructs - 1;
+			}
+			
+			return INDEX_NONE;
+		}
+		
+		const UScriptStruct* ScriptStructs[MaxScriptStructsPerChunk];	/** Types ScriptStructs stored in this chunk. */
+		int32 NumScriptStructs;											/** Number of ScriptsStructs */
+		int32 UsedSize;													/** Size of the used memory in the chunk. */
+		int32 NumItems;													/** Number of items in the chunks */
+		FChunkHeader* Next;												/** Pointer to next chunk. */
+	};
+
+	/** Allocates a new chunk and makes it the head of the chunk list. */
+	void AllocateNewChunk();
+
+	/** Allocates an item of specified type. */
+	uint8* AllocItem(const UScriptStruct* InScriptStruct)
+	{
+		check(InScriptStruct);
+
+		const int32 MinAlignment = InScriptStruct->GetMinAlignment();
+		const int32 ScriptStructSize = InScriptStruct->GetStructureSize();
+		check((Align(sizeof(FItemHeader), MinAlignment) + ScriptStructSize) < ChunkSize);
+
+		// Make sure there's a chunk prsent.
+		if (Head == nullptr)
+		{
+			AllocateNewChunk();
+		}
+		check(Head != nullptr);
+
+		// Make sure there's space for the struct type and allocation
+		int32 ScriptStructIndex = Head->GetScriptStructIndex(InScriptStruct);
+		int32 HeaderOffset = Head->UsedSize;
+		int32 ItemOffset = Align(Head->UsedSize + sizeof(FItemHeader), MinAlignment);
+		
+		if (ScriptStructIndex == INDEX_NONE || (ItemOffset + ScriptStructSize) > ChunkSize)
+		{
+			AllocateNewChunk();
+			ScriptStructIndex = Head->GetScriptStructIndex(InScriptStruct);
+			HeaderOffset = Head->UsedSize;
+			ItemOffset = Align(Head->UsedSize + sizeof(FItemHeader), MinAlignment);
+			check(ScriptStructIndex != INDEX_NONE);
+			check((ItemOffset + ScriptStructSize) <= ChunkSize)
+		}
+
+		uint8* ChunkMemory = reinterpret_cast<uint8*>(Head) + sizeof(FChunkHeader);
+		Head->UsedSize = ItemOffset + ScriptStructSize;
+		Head->NumItems++;
+
+		FItemHeader& ItemHeader = *reinterpret_cast<FItemHeader*>(ChunkMemory + HeaderOffset);
+		ItemHeader.Size = (ItemOffset + ScriptStructSize) - HeaderOffset;
+		ItemHeader.Offset = ItemOffset - HeaderOffset;
+		ItemHeader.ScriptStructIndex = ScriptStructIndex;
+
+		NumItems++;
+		
+		return ChunkMemory + ItemOffset;
+	}
+
+	FChunkHeader* First = nullptr;		/** Pointer to the first chunk in use */
+	FChunkHeader* Head = nullptr;		/** Pointer to the last chunk in use */
+	FChunkHeader* FreeList = nullptr;	/** Free list of unused chunks */
+	int32 NumItems = 0;					/** Number of items stored in the buffer */
+	int32 ChunkSize = 0;				/** Size of each chunk */
+};
+
+
+template<>
+struct TStructOpsTypeTraits<FChunkedStructBuffer> : public TStructOpsTypeTraitsBase2<FChunkedStructBuffer>
+{
+	enum
+	{
+		WithAddStructReferencedObjects = true,
+	};
+};
\ No newline at end of file
diff --git a/Engine/Plugins/Experimental/StructUtils/Source/StructUtils/Public/InstancedStructStream.h b/Engine/Plugins/Experimental/StructUtils/Source/StructUtils/Public/InstancedStructStream.h
new file mode 100644
--- /dev/null
+++ b/Engine/Plugins/Experimental/StructUtils/Source/StructUtils/Public/InstancedStructStream.h
@@ -0,0 +1,392 @@
+﻿// Copyright Epic Games, Inc. All Rights Reserved.
+
+#pragma once
+
+#include "CoreMinimal.h"
+#include "UObject/Class.h"
+#include "InstancedStruct.h"
+#include "InstancedStructStream.generated.h"
+
+/**
+ * A stream where you can append heterogeneous structs, and iterate over them in order (no random access).
+ * The structs are stored in chunks of specified size, and reused when the stream is reset.
+ * Each chunk can hold up to MaxScriptStructsPerChunk types, on overflow new chunk is allocated.
+ * The stream can be iterated over using one of the ForEach*() methods.
+ * Note: If you use move assign or Append(&&), make sure to call Compact() or you may end up with large free list.
+ */
+USTRUCT()
+struct STRUCTUTILS_API FInstancedStructStream
+{
+	GENERATED_BODY()
+
+	static constexpr int32 DefaultChunkSize = 16384;
+
+	FInstancedStructStream(const int32 InChunkSize = DefaultChunkSize)
+		: ChunkSize(InChunkSize)
+	{
+	}
+
+	FInstancedStructStream(const FInstancedStructStream& Other)
+	{
+		ChunkSize = Other.ChunkSize;
+		Append(Other);
+	}
+	
+	FInstancedStructStream(FInstancedStructStream&& Other)
+	{
+		ChunkSize = Other.ChunkSize;
+		Append(Other);
+	}
+
+	~FInstancedStructStream()
+	{
+		Clear();
+	}
+
+	FInstancedStructStream& operator=(const FInstancedStructStream& Other)
+	{
+		Reset();
+		Append(Other);
+		return *this;
+	}
+	
+	FInstancedStructStream& operator=(FInstancedStructStream&& Other)
+	{
+		Reset();
+		Append(Other);
+		return *this;
+	}
+
+	void Append(const FInstancedStructStream& Other)
+	{
+		Other.ForEach([this](FStructView View)
+		{
+			Add(View);
+		});
+	}
+	
+	void Append(FInstancedStructStream&& Other)
+	{
+		if (Other.ChunkSize == ChunkSize)
+		{
+			// If same chunk size, steal chunks.
+			NumItems += Other.NumItems;
+			Head->Next = Other.First;
+			Head = Other.Head;
+
+			Other.NumItems = 0;
+			Other.First = nullptr;
+			Other.Head = nullptr;
+		}
+		else
+		{
+			Other.ForEach([this](FStructView View)
+			{
+				Add(View);
+			});
+		}
+		Other.Clear();
+	}
+
+	/** Emplaces struct in the buffer. */
+	template<typename T, typename... TArgs>
+	void Emplace(TArgs&&... InArgs)
+	{
+		const UScriptStruct* InScriptStruct = T::StaticStruct();
+		uint8* ItemMemory = AllocItem(InScriptStruct);
+		new (ItemMemory) T(Forward<TArgs>(InArgs)...);
+	}
+
+	/** Emplaces struct in the buffer, and returns reference to it. */
+	template<typename T, typename... TArgs>
+	T& Emplace_GetRef(TArgs&&... InArgs)
+	{
+		const UScriptStruct* InScriptStruct = T::StaticStruct();
+		uint8* ItemMemory = AllocItem(InScriptStruct);
+		return *new (ItemMemory) T(Forward<TArgs>(InArgs)...);
+	}
+
+	/** Add new struct in the buffer. */
+	template<typename T>
+	void Add(const T& InStruct)
+	{
+		const UScriptStruct* InScriptStruct = T::StaticStruct();
+		uint8* ItemMemory = AllocItem(InScriptStruct);
+		new (ItemMemory) T(InStruct);
+	}
+
+	/** Add new struct in the buffer, and return reference to it. */
+	template<typename T>
+	T& Add_GetRef(const T& InStruct)
+	{
+		const UScriptStruct* InScriptStruct = T::StaticStruct();
+		uint8* ItemMemory = AllocItem(InScriptStruct);
+		return *new (ItemMemory) T(InStruct);
+	}
+
+	/** Adds struct to the buffer based on ScriptStruct and pointer. */
+	void Add(const FStructView Struct)
+	{
+		check(Struct.IsValid());
+		const UScriptStruct* InScriptStruct = Struct.GetScriptStruct();
+		uint8* ItemMemory = AllocItem(InScriptStruct);
+		InScriptStruct->InitializeStruct(ItemMemory);
+		InScriptStruct->CopyScriptStruct(ItemMemory, Struct.GetMemory());
+	}
+
+	/** Iterates over all structs and calls Function on each item */
+	template<typename TFunc>
+	void ForEach(TFunc&& Function) const
+	{
+		for (FChunkHeader* Chunk = First; Chunk != nullptr; Chunk = Chunk->Next)
+		{
+			uint8* ItemMemory = reinterpret_cast<uint8*>(Chunk) + sizeof(FChunkHeader);
+			for (int32 Index = 0; Index < Chunk->NumItems; Index++)
+			{
+				FItemHeader& Item = *reinterpret_cast<FItemHeader*>(ItemMemory);
+				check(Item.ScriptStructIndex < MaxScriptStructsPerChunk);
+			
+				Function(FStructView(Chunk->ScriptStructs[Item.ScriptStructIndex], ItemMemory + Item.Offset));
+
+				ItemMemory += Item.Size;
+			}
+		}
+	}
+
+	/**
+	 * Iterates over all structs of specified type T and calls Function on each item.
+	 * Usage: Buffer.ForEach<FFoo>([](Foo& Item) { ... });
+	 */
+	template<typename T, typename TFunc>
+	void ForEach(TFunc&& Function) const
+	{
+		for (FChunkHeader* Chunk = First; Chunk != nullptr; Chunk = Chunk->Next)
+		{
+			// Build mask of the accepted struct in the chunk.
+			const int32 StructIndex = Chunk->GetScriptStructIndex(T::StaticStruct());
+			if (StructIndex == INDEX_NONE)
+			{
+				continue;
+			}
+			const uint32 FilterMask = 1U << StructIndex;
+			
+			uint8* ItemMemory = reinterpret_cast<uint8*>(Chunk) + sizeof(FChunkHeader);
+			for (int32 Index = 0; Index < Chunk->NumItems; Index++)
+			{
+				FItemHeader& Item = *reinterpret_cast<FItemHeader*>(ItemMemory);
+				const uint32 StructMask = 1U << Item.ScriptStructIndex;
+
+				if ((FilterMask & StructMask) != 0)
+				{
+					check(Item.ScriptStructIndex < MaxScriptStructsPerChunk);
+					Function(*reinterpret_cast<T*>(ItemMemory + Item.Offset));
+				}
+
+				ItemMemory += Item.Size;
+			}
+		}
+	}
+
+	/** Iterates over all structs of specified types and calls Function on each item */
+	template<typename TFunc>
+	void ForEachFiltered(TArrayView<const UScriptStruct*> AcceptedScriptStructs, TFunc&& Function) const
+	{
+		for (FChunkHeader* Chunk = First; Chunk != nullptr; Chunk = Chunk->Next)
+		{
+			// Build mask of all accepted structs in the chunk.
+			uint32 FilterMask = 0;
+			for (const UScriptStruct* ScriptStruct : AcceptedScriptStructs)
+			{
+				const int32 StructIndex = Chunk->GetScriptStructIndex(ScriptStruct);
+				if (StructIndex != INDEX_NONE)
+				{
+					FilterMask |= 1U << StructIndex;
+				}
+			}
+
+			// Chunk does not match any of the structs specified in the filter.
+			if (FilterMask == 0)
+			{
+				continue;
+			}
+			
+			uint8* ItemMemory = reinterpret_cast<uint8*>(Chunk) + sizeof(FChunkHeader);
+			for (int32 Index = 0; Index < Chunk->NumItems; Index++)
+			{
+				FItemHeader& Item = *reinterpret_cast<FItemHeader*>(ItemMemory);
+				const uint32 StructMask = 1U << Item.ScriptStructIndex;
+
+				if ((FilterMask & StructMask) != 0)
+				{
+					check(Item.ScriptStructIndex < MaxScriptStructsPerChunk);
+					Function(FStructView(Chunk->ScriptStructs[Item.ScriptStructIndex], ItemMemory + Item.Offset));
+				}
+
+				ItemMemory += Item.Size;
+			}
+		}
+	}
+
+	/** Returns struct types in the buffer */
+	void GetScriptStructs(TArray<const UScriptStruct*>& OutScriptStructs) const
+	{
+		OutScriptStructs.Reset();
+		for (const FChunkHeader* Chunk = First; Chunk != nullptr; Chunk = Chunk->Next)
+		{
+			// Add reference to the ScriptStruct object.
+			for (int32 Index = 0; Index < Chunk->NumScriptStructs; Index++)
+			{
+				OutScriptStructs.AddUnique(Chunk->ScriptStructs[Index]);
+			}
+		}
+	}
+
+	/** Resets and clears all structs, keeps internal memory. */
+	void Reset();
+
+	/** Releases unused internal memory. */
+	void Compact();
+
+	/** Resets and clears all structs, and releases unused internal memory. */
+	void Clear();
+
+	/** @return Number of structs added to the buffer. */
+	int32 Num() const { return NumItems; }
+
+	/** @return True if the buffer is empty. */
+	bool IsEmpty() const { return NumItems == 0; }
+
+	/** @return Chunk size */
+	int32 GetChunkSize() const { return ChunkSize; }
+
+	/** @return Number of chunks in use. */
+	int32 GetNumUsedChunks() const
+	{
+		int32 Count = 0;
+		for (const FChunkHeader* Chunk = First; Chunk != nullptr; Chunk = Chunk->Next)
+		{
+			Count++;
+		}
+		return Count;
+	}
+
+	/** @return Number of  chunks in freelist. */
+	int32 GetNumFreeChunks() const
+	{
+		int32 Count = 0;
+		for (const FChunkHeader* Chunk = FreeList; Chunk != nullptr; Chunk = Chunk->Next)
+		{
+			Count++;
+		}
+		return Count;
+	}
+
+	void AddStructReferencedObjects(class FReferenceCollector& Collector);
+
+protected:
+
+	/** Number of unique struct types that can be stored per chunk. */
+	static constexpr int32 MaxScriptStructsPerChunk = 8;
+
+	/** Header for each item in a chunk. Allocated before the item in tge chunk memory. */
+	struct FItemHeader
+	{
+		int32 Size;					/** Size of the item including this header */
+		int16 Offset;				/** relative offset to the item data */ 
+		int16 ScriptStructIndex;	/** Index in chunk ScriptStruct array to the type of the item */ 
+	};
+
+	/** Header for a chunk of items. */
+	struct FChunkHeader
+	{
+		/** @returns Index to the ScriptStruct stored in the chunk, or INDEX_NONE if cannot store the ScriptStruct. */
+		int32 GetScriptStructIndex(const UScriptStruct* InScriptStruct)
+		{
+			for (int32 Index = 0; Index < NumScriptStructs; Index++)
+			{
+				if (ScriptStructs[Index] == InScriptStruct)
+				{
+					return Index;
+				}
+			}
+
+			if (NumScriptStructs < MaxScriptStructsPerChunk)
+			{
+				ScriptStructs[NumScriptStructs++] = InScriptStruct;
+				return NumScriptStructs - 1;
+			}
+			
+			return INDEX_NONE;
+		}
+		
+		const UScriptStruct* ScriptStructs[MaxScriptStructsPerChunk];	/** Types ScriptStructs stored in this chunk. */
+		int32 NumScriptStructs;											/** Number of ScriptsStructs */
+		int32 UsedSize;													/** Size of the used memory in the chunk. */
+		int32 NumItems;													/** Number of items in the chunks */
+		FChunkHeader* Next;												/** Pointer to next chunk. */
+	};
+
+	/** Allocates a new chunk and makes it the head of the chunk list. */
+	void AllocateNewChunk();
+
+	/** Allocates an item of specified type. */
+	uint8* AllocItem(const UScriptStruct* InScriptStruct)
+	{
+		check(InScriptStruct);
+
+		const int32 MinAlignment = InScriptStruct->GetMinAlignment();
+		const int32 ScriptStructSize = InScriptStruct->GetStructureSize();
+		check((Align(sizeof(FItemHeader), MinAlignment) + ScriptStructSize) < ChunkSize);
+
+		// Make sure there's a chunk prsent.
+		if (Head == nullptr)
+		{
+			AllocateNewChunk();
+		}
+		check(Head != nullptr);
+
+		// Make sure there's space for the struct type and allocation
+		int32 ScriptStructIndex = Head->GetScriptStructIndex(InScriptStruct);
+		int32 HeaderOffset = Head->UsedSize;
+		int32 ItemOffset = Align(Head->UsedSize + sizeof(FItemHeader), MinAlignment);
+		
+		if (ScriptStructIndex == INDEX_NONE || (ItemOffset + ScriptStructSize) > ChunkSize)
+		{
+			AllocateNewChunk();
+			ScriptStructIndex = Head->GetScriptStructIndex(InScriptStruct);
+			HeaderOffset = Head->UsedSize;
+			ItemOffset = Align(Head->UsedSize + sizeof(FItemHeader), MinAlignment);
+			check(ScriptStructIndex != INDEX_NONE);
+			check((ItemOffset + ScriptStructSize) <= ChunkSize)
+		}
+
+		uint8* ChunkMemory = reinterpret_cast<uint8*>(Head) + sizeof(FChunkHeader);
+		Head->UsedSize = ItemOffset + ScriptStructSize;
+		Head->NumItems++;
+
+		FItemHeader& ItemHeader = *reinterpret_cast<FItemHeader*>(ChunkMemory + HeaderOffset);
+		ItemHeader.Size = (ItemOffset + ScriptStructSize) - HeaderOffset;
+		ItemHeader.Offset = ItemOffset - HeaderOffset;
+		ItemHeader.ScriptStructIndex = ScriptStructIndex;
+
+		NumItems++;
+		
+		return ChunkMemory + ItemOffset;
+	}
+
+	FChunkHeader* First = nullptr;		/** Pointer to the first chunk in use */
+	FChunkHeader* Head = nullptr;		/** Pointer to the last chunk in use */
+	FChunkHeader* FreeList = nullptr;	/** Free list of unused chunks */
+	int32 NumItems = 0;					/** Number of items stored in the buffer */
+	int32 ChunkSize = 0;				/** Size of each chunk */
+};
+
+
+template<>
+struct TStructOpsTypeTraits<FInstancedStructStream> : public TStructOpsTypeTraitsBase2<FInstancedStructStream>
+{
+	enum
+	{
+		WithAddStructReferencedObjects = true,
+	};
+};
\ No newline at end of file
diff --git a/Engine/Plugins/Experimental/StructUtils/Source/StructUtilsTestSuite/Private/ChunkedStructBufferTest.cpp b/Engine/Plugins/Experimental/StructUtils/Source/StructUtilsTestSuite/Private/ChunkedStructBufferTest.cpp
new file mode 100644
--- /dev/null
+++ b/Engine/Plugins/Experimental/StructUtils/Source/StructUtilsTestSuite/Private/ChunkedStructBufferTest.cpp
@@ -0,0 +1,202 @@
+// Copyright Epic Games, Inc. All Rights Reserved.
+
+#include "CoreMinimal.h"
+#include "AITestsCommon.h"
+#include "ChunkedStructBuffer.h"
+#include "StructUtilsTestTypes.h"
+
+#define LOCTEXT_NAMESPACE "StructUtilsTests"
+
+PRAGMA_DISABLE_OPTIMIZATION
+
+struct FStructUtilsTest_ChunkedStructBufferBasic : FAITestBase
+{
+	virtual bool InstantTest() override
+	{
+		FChunkedStructBuffer Buffer(1234);
+
+		AITEST_TRUE("Chunk size should be 1234", Buffer.GetChunkSize() == 1234);
+
+		Buffer.Emplace<FTestStructSimple>(42.0f);
+		AITEST_TRUE("Should have 1 item", Buffer.Num() == 1);
+
+		FTestStructComplex Item;
+		Item.String = TEXT("Foo");
+		Item.StringArray.Add(FString(TEXT("Bar")));
+		Buffer.Add(Item);
+		AITEST_TRUE("Should have 2 items", Buffer.Num() == 2);
+
+		int Index = 0;
+		bool bFoundSimple = false;
+		bool bFoundComplex = false;
+		Buffer.ForEach([&Index, &bFoundSimple, &bFoundComplex](FStructView Item)
+		{
+			if (Index == 0)
+			{
+				bFoundSimple = Item.GetScriptStruct() == FTestStructSimple::StaticStruct();
+			}
+			else if (Index == 1)
+			{
+				bFoundComplex = Item.GetScriptStruct() == FTestStructComplex::StaticStruct();
+			}
+
+			Index++;
+		});
+
+		AITEST_TRUE("Should have FTestStructSimple at index 0", bFoundSimple);
+		AITEST_TRUE("Should have FTestStructComplex at index 1", bFoundComplex);
+
+		Buffer.Clear();
+		AITEST_TRUE("Should have 0 items", Buffer.Num() == 0);
+		AITEST_TRUE("Should have 0 chunks", Buffer.GetNumUsedChunks() == 0);
+		AITEST_TRUE("Should have 0 free chunk", Buffer.GetNumFreeChunks() == 0);
+
+		// This should allocate second chunk,as there are more than 8 types.
+		Buffer.Emplace<FTestStructSimple>();
+		Buffer.Emplace<FTestStructSimple1>();
+		Buffer.Emplace<FTestStructSimple2>();
+		Buffer.Emplace<FTestStructSimple3>();
+		Buffer.Emplace<FTestStructSimple4>();
+		Buffer.Emplace<FTestStructSimple5>();
+		Buffer.Emplace<FTestStructSimple6>();
+		Buffer.Emplace<FTestStructSimple7>();
+		Buffer.Emplace<FTestStructComplex>();
+		AITEST_TRUE("Should have 9 items", Buffer.Num() == 9);
+		AITEST_TRUE("Should have 2 chunks", Buffer.GetNumUsedChunks() == 2);
+		
+		return true;
+	}
+};
+IMPLEMENT_AI_INSTANT_TEST(FStructUtilsTest_ChunkedStructBufferBasic, "System.StructUtils.ChunkedStructBuffer.Basic");
+
+struct FStructUtilsTest_ChunkedStructBufferAppend : FAITestBase
+{
+	virtual bool InstantTest() override
+	{
+		FChunkedStructBuffer Buffer1(64);
+		FChunkedStructBuffer Buffer2(64);
+		FChunkedStructBuffer Buffer3(64);
+
+		for (int32 i = 0; i < 30; i++)
+		{
+			FTestStructSimple& Item = Buffer1.Emplace_GetRef<FTestStructSimple>();
+			Item.Float = (float)i;
+		}
+		AITEST_TRUE("Buffer1 should have 30 item", Buffer1.Num() == 30);
+		
+		for (int32 i = 0; i < 5; i++)
+		{
+			FTestStructComplex Item;
+			FTestStructComplex& ItemRef = Buffer2.Add_GetRef(Item);
+			ItemRef.String = FString::FormatAsNumber(i);
+		}
+		AITEST_TRUE("Buffer2 should have 5 item", Buffer2.Num() == 5);
+
+		Buffer3.Append(Buffer1);
+		AITEST_TRUE("Buffer3 should have 30 items", Buffer3.Num() == 30);
+
+		Buffer3.Append(MoveTemp(Buffer2));
+		AITEST_TRUE("Buffer3 should have 35 items", Buffer3.Num() == 35);
+		AITEST_TRUE("Buffer2 should have 0 items", Buffer2.Num() == 0);
+		AITEST_TRUE("Buffer2 should have 0 used chunks", Buffer2.GetNumUsedChunks() == 0);
+		AITEST_TRUE("Buffer2 should have 0 free chunks", Buffer2.GetNumFreeChunks() == 0);
+
+		Buffer2.Emplace<FTestStructSimple>(1);
+		AITEST_TRUE("Buffer2 should have 1 item", Buffer2.Num() == 1);
+
+		return true;
+	}
+};
+IMPLEMENT_AI_INSTANT_TEST(FStructUtilsTest_ChunkedStructBufferAppend, "System.StructUtils.ChunkedStructBuffer.Append");
+
+
+struct FStructUtilsTest_ChunkedStructBufferFreelist : FAITestBase
+{
+	virtual bool InstantTest() override
+	{
+		FChunkedStructBuffer Buffer(64);
+
+		for (int32 i = 0; i < 30; i++)
+		{
+			FTestStructSimple& Item = Buffer.Emplace_GetRef<FTestStructSimple>();
+			Item.Float = (float)i;
+		}
+		AITEST_TRUE("Buffer should have 30 item", Buffer.Num() == 30);
+
+		Buffer.Reset();
+		AITEST_TRUE("Buffer should have 0 item", Buffer.Num() == 0);
+		AITEST_TRUE("Should have 0 chunks", Buffer.GetNumUsedChunks() == 0);
+		AITEST_TRUE("Should have >0 free chunks", Buffer.GetNumFreeChunks() > 0);
+
+		for (int32 i = 0; i < 45; i++)
+		{
+			Buffer.Emplace<FTestStructSimple>((float)-i);
+		}
+		AITEST_TRUE("Buffer should have 45 items", Buffer.Num() == 45);
+		AITEST_TRUE("Should have 0 free chunks", Buffer.GetNumFreeChunks() == 0);
+
+		return true;
+	}
+};
+IMPLEMENT_AI_INSTANT_TEST(FStructUtilsTest_ChunkedStructBufferFreelist, "System.StructUtils.ChunkedStructBuffer.Freelist");
+
+
+struct FStructUtilsTest_ChunkedStructBufferForEach : FAITestBase
+{
+	virtual bool InstantTest() override
+	{
+		FChunkedStructBuffer Buffer;
+
+		Buffer.Emplace<FTestStructSimple>(1.0f);
+		Buffer.Emplace<FTestStructComplex>();
+		Buffer.Emplace<FTestStructSimple>(2.0f);
+		Buffer.Emplace<FTestStructSimple>(3.0f);
+		Buffer.Emplace<FTestStructSimple1>();
+		Buffer.Emplace<FTestStructComplex>();
+
+		int32 SimpleCount = 0;
+		int32 ComplexCount = 0;
+		Buffer.ForEach([&SimpleCount, &ComplexCount](FStructView Item)
+		{
+			if (Item.GetScriptStruct() == FTestStructSimple::StaticStruct())
+			{
+				SimpleCount++;
+			}
+			else if (Item.GetScriptStruct() == FTestStructComplex::StaticStruct())
+			{
+				ComplexCount++;
+			}
+		});
+		AITEST_TRUE("Buffer should have 3 simple items", SimpleCount == 3);
+		AITEST_TRUE("Buffer should have 2 complex items", ComplexCount == 2);
+
+	
+		int32 Filter1Count = 0;
+		float Sum = 0.0f;
+		Buffer.ForEach<FTestStructSimple>([&Filter1Count, &Sum](FTestStructSimple& Item)
+		{
+			Sum += Item.Float;
+			Filter1Count++;
+		});
+		AITEST_TRUE("Should have 3 items with Filter1", Filter1Count == 3);
+		AITEST_TRUE("Should have sum of 6", FMath::IsNearlyEqual(Sum, 6.0f));
+
+		TArray<const UScriptStruct*> Filter2 { FTestStructSimple::StaticStruct(), FTestStructSimple1::StaticStruct() };
+		int32 Filter2Count = 0;
+		Buffer.ForEachFiltered(Filter2, [&Filter2Count](FStructView Item)
+		{
+			Filter2Count++;
+		});
+		AITEST_TRUE("Should have 4 items with Filter2", Filter2Count == 4);
+
+		
+		return true;
+	}
+};
+IMPLEMENT_AI_INSTANT_TEST(FStructUtilsTest_ChunkedStructBufferForEach, "System.StructUtils.ChunkedStructBuffer.ForEach");
+
+
+
+PRAGMA_ENABLE_OPTIMIZATION
+
+#undef LOCTEXT_NAMESPACE
diff --git a/Engine/Plugins/Experimental/StructUtils/Source/StructUtilsTestSuite/Private/InstancedStructStreamTest.cpp b/Engine/Plugins/Experimental/StructUtils/Source/StructUtilsTestSuite/Private/InstancedStructStreamTest.cpp
new file mode 100644
--- /dev/null
+++ b/Engine/Plugins/Experimental/StructUtils/Source/StructUtilsTestSuite/Private/InstancedStructStreamTest.cpp
@@ -0,0 +1,202 @@
+// Copyright Epic Games, Inc. All Rights Reserved.
+
+#include "CoreMinimal.h"
+#include "AITestsCommon.h"
+#include "InstancedStructStream.h"
+#include "StructUtilsTestTypes.h"
+
+#define LOCTEXT_NAMESPACE "StructUtilsTests"
+
+PRAGMA_DISABLE_OPTIMIZATION
+
+struct FStructUtilsTest_InstancedStructStreamBasic : FAITestBase
+{
+	virtual bool InstantTest() override
+	{
+		FInstancedStructStream Buffer(1234);
+
+		AITEST_TRUE("Chunk size should be 1234", Buffer.GetChunkSize() == 1234);
+
+		Buffer.Emplace<FTestStructSimple>(42.0f);
+		AITEST_TRUE("Should have 1 item", Buffer.Num() == 1);
+
+		FTestStructComplex Item;
+		Item.String = TEXT("Foo");
+		Item.StringArray.Add(FString(TEXT("Bar")));
+		Buffer.Add(Item);
+		AITEST_TRUE("Should have 2 items", Buffer.Num() == 2);
+
+		int Index = 0;
+		bool bFoundSimple = false;
+		bool bFoundComplex = false;
+		Buffer.ForEach([&Index, &bFoundSimple, &bFoundComplex](FStructView Item)
+		{
+			if (Index == 0)
+			{
+				bFoundSimple = Item.GetScriptStruct() == FTestStructSimple::StaticStruct();
+			}
+			else if (Index == 1)
+			{
+				bFoundComplex = Item.GetScriptStruct() == FTestStructComplex::StaticStruct();
+			}
+
+			Index++;
+		});
+
+		AITEST_TRUE("Should have FTestStructSimple at index 0", bFoundSimple);
+		AITEST_TRUE("Should have FTestStructComplex at index 1", bFoundComplex);
+
+		Buffer.Clear();
+		AITEST_TRUE("Should have 0 items", Buffer.Num() == 0);
+		AITEST_TRUE("Should have 0 chunks", Buffer.GetNumUsedChunks() == 0);
+		AITEST_TRUE("Should have 0 free chunk", Buffer.GetNumFreeChunks() == 0);
+
+		// This should allocate second chunk,as there are more than 8 types.
+		Buffer.Emplace<FTestStructSimple>();
+		Buffer.Emplace<FTestStructSimple1>();
+		Buffer.Emplace<FTestStructSimple2>();
+		Buffer.Emplace<FTestStructSimple3>();
+		Buffer.Emplace<FTestStructSimple4>();
+		Buffer.Emplace<FTestStructSimple5>();
+		Buffer.Emplace<FTestStructSimple6>();
+		Buffer.Emplace<FTestStructSimple7>();
+		Buffer.Emplace<FTestStructComplex>();
+		AITEST_TRUE("Should have 9 items", Buffer.Num() == 9);
+		AITEST_TRUE("Should have 2 chunks", Buffer.GetNumUsedChunks() == 2);
+		
+		return true;
+	}
+};
+IMPLEMENT_AI_INSTANT_TEST(FStructUtilsTest_InstancedStructStreamBasic, "System.StructUtils.InstancedStructStream.Basic");
+
+struct FStructUtilsTest_InstancedStructStreamAppend : FAITestBase
+{
+	virtual bool InstantTest() override
+	{
+		FInstancedStructStream Buffer1(64);
+		FInstancedStructStream Buffer2(64);
+		FInstancedStructStream Buffer3(64);
+
+		for (int32 i = 0; i < 30; i++)
+		{
+			FTestStructSimple& Item = Buffer1.Emplace_GetRef<FTestStructSimple>();
+			Item.Float = (float)i;
+		}
+		AITEST_TRUE("Buffer1 should have 30 item", Buffer1.Num() == 30);
+		
+		for (int32 i = 0; i < 5; i++)
+		{
+			FTestStructComplex Item;
+			FTestStructComplex& ItemRef = Buffer2.Add_GetRef(Item);
+			ItemRef.String = FString::FormatAsNumber(i);
+		}
+		AITEST_TRUE("Buffer2 should have 5 item", Buffer2.Num() == 5);
+
+		Buffer3.Append(Buffer1);
+		AITEST_TRUE("Buffer3 should have 30 items", Buffer3.Num() == 30);
+
+		Buffer3.Append(MoveTemp(Buffer2));
+		AITEST_TRUE("Buffer3 should have 35 items", Buffer3.Num() == 35);
+		AITEST_TRUE("Buffer2 should have 0 items", Buffer2.Num() == 0);
+		AITEST_TRUE("Buffer2 should have 0 used chunks", Buffer2.GetNumUsedChunks() == 0);
+		AITEST_TRUE("Buffer2 should have 0 free chunks", Buffer2.GetNumFreeChunks() == 0);
+
+		Buffer2.Emplace<FTestStructSimple>(1);
+		AITEST_TRUE("Buffer2 should have 1 item", Buffer2.Num() == 1);
+
+		return true;
+	}
+};
+IMPLEMENT_AI_INSTANT_TEST(FStructUtilsTest_InstancedStructStreamAppend, "System.StructUtils.InstancedStructStream.Append");
+
+
+struct FStructUtilsTest_InstancedStructStreamFreelist : FAITestBase
+{
+	virtual bool InstantTest() override
+	{
+		FInstancedStructStream Buffer(64);
+
+		for (int32 i = 0; i < 30; i++)
+		{
+			FTestStructSimple& Item = Buffer.Emplace_GetRef<FTestStructSimple>();
+			Item.Float = (float)i;
+		}
+		AITEST_TRUE("Buffer should have 30 item", Buffer.Num() == 30);
+
+		Buffer.Reset();
+		AITEST_TRUE("Buffer should have 0 item", Buffer.Num() == 0);
+		AITEST_TRUE("Should have 0 chunks", Buffer.GetNumUsedChunks() == 0);
+		AITEST_TRUE("Should have >0 free chunks", Buffer.GetNumFreeChunks() > 0);
+
+		for (int32 i = 0; i < 45; i++)
+		{
+			Buffer.Emplace<FTestStructSimple>((float)-i);
+		}
+		AITEST_TRUE("Buffer should have 45 items", Buffer.Num() == 45);
+		AITEST_TRUE("Should have 0 free chunks", Buffer.GetNumFreeChunks() == 0);
+
+		return true;
+	}
+};
+IMPLEMENT_AI_INSTANT_TEST(FStructUtilsTest_InstancedStructStreamFreelist, "System.StructUtils.InstancedStructStream.Freelist");
+
+
+struct FStructUtilsTest_InstancedStructStreamForEach : FAITestBase
+{
+	virtual bool InstantTest() override
+	{
+		FInstancedStructStream Buffer;
+
+		Buffer.Emplace<FTestStructSimple>(1.0f);
+		Buffer.Emplace<FTestStructComplex>();
+		Buffer.Emplace<FTestStructSimple>(2.0f);
+		Buffer.Emplace<FTestStructSimple>(3.0f);
+		Buffer.Emplace<FTestStructSimple1>();
+		Buffer.Emplace<FTestStructComplex>();
+
+		int32 SimpleCount = 0;
+		int32 ComplexCount = 0;
+		Buffer.ForEach([&SimpleCount, &ComplexCount](FStructView Item)
+		{
+			if (Item.GetScriptStruct() == FTestStructSimple::StaticStruct())
+			{
+				SimpleCount++;
+			}
+			else if (Item.GetScriptStruct() == FTestStructComplex::StaticStruct())
+			{
+				ComplexCount++;
+			}
+		});
+		AITEST_TRUE("Buffer should have 3 simple items", SimpleCount == 3);
+		AITEST_TRUE("Buffer should have 2 complex items", ComplexCount == 2);
+
+	
+		int32 Filter1Count = 0;
+		float Sum = 0.0f;
+		Buffer.ForEach<FTestStructSimple>([&Filter1Count, &Sum](FTestStructSimple& Item)
+		{
+			Sum += Item.Float;
+			Filter1Count++;
+		});
+		AITEST_TRUE("Should have 3 items with Filter1", Filter1Count == 3);
+		AITEST_TRUE("Should have sum of 6", FMath::IsNearlyEqual(Sum, 6.0f));
+
+		TArray<const UScriptStruct*> Filter2 { FTestStructSimple::StaticStruct(), FTestStructSimple1::StaticStruct() };
+		int32 Filter2Count = 0;
+		Buffer.ForEachFiltered(Filter2, [&Filter2Count](FStructView Item)
+		{
+			Filter2Count++;
+		});
+		AITEST_TRUE("Should have 4 items with Filter2", Filter2Count == 4);
+
+		
+		return true;
+	}
+};
+IMPLEMENT_AI_INSTANT_TEST(FStructUtilsTest_InstancedStructStreamForEach, "System.StructUtils.InstancedStructStream.ForEach");
+
+
+
+PRAGMA_ENABLE_OPTIMIZATION
+
+#undef LOCTEXT_NAMESPACE
diff --git a/Engine/Plugins/Experimental/StructUtils/Source/StructUtilsTestSuite/Private/StructUtilsTestSuiteModule.cpp b/Engine/Plugins/Experimental/StructUtils/Source/StructUtilsTestSuite/Private/StructUtilsTestSuiteModule.cpp
new file mode 100644
--- /dev/null
+++ b/Engine/Plugins/Experimental/StructUtils/Source/StructUtilsTestSuite/Private/StructUtilsTestSuiteModule.cpp
@@ -0,0 +1,13 @@
+// Copyright Epic Games, Inc. All Rights Reserved.
+
+#include "StructUtilsTestSuiteModule.h"
+
+#define LOCTEXT_NAMESPACE "StructUtilsTestSuite"
+
+class FStructUtilsTestSuiteModule : public IStructUtilsTestSuiteModule
+{
+};
+
+IMPLEMENT_MODULE(FStructUtilsTestSuiteModule, StructUtilsTestSuite)
+
+#undef LOCTEXT_NAMESPACE
diff --git a/Engine/Plugins/Experimental/StructUtils/Source/StructUtilsTestSuite/Private/StructUtilsTestTypes.cpp b/Engine/Plugins/Experimental/StructUtils/Source/StructUtilsTestSuite/Private/StructUtilsTestTypes.cpp
new file mode 100644
--- /dev/null
+++ b/Engine/Plugins/Experimental/StructUtils/Source/StructUtilsTestSuite/Private/StructUtilsTestTypes.cpp
@@ -0,0 +1 @@
+// Copyright Epic Games, Inc. All Rights Reserved.
diff --git a/Engine/Plugins/Experimental/StructUtils/Source/StructUtilsTestSuite/Public/StructUtilsTestSuiteModule.h b/Engine/Plugins/Experimental/StructUtils/Source/StructUtilsTestSuite/Public/StructUtilsTestSuiteModule.h
new file mode 100644
--- /dev/null
+++ b/Engine/Plugins/Experimental/StructUtils/Source/StructUtilsTestSuite/Public/StructUtilsTestSuiteModule.h
@@ -0,0 +1,37 @@
+// Copyright Epic Games, Inc. All Rights Reserved.
+
+#pragma once
+
+#include "CoreMinimal.h"
+#include "Modules/ModuleManager.h"
+
+/**
+* The public interface to this module
+*/
+class IStructUtilsTestSuiteModule : public IModuleInterface
+{
+
+public:
+
+	/**
+	* Singleton-like access to this module's interface.  This is just for convenience!
+	* Beware of calling this during the shutdown phase, though.  Your module might have been unloaded already.
+	*
+	* @return Returns singleton instance, loading the module on demand if needed
+	*/
+	static inline IStructUtilsTestSuiteModule& Get()
+	{
+		return FModuleManager::LoadModuleChecked<IStructUtilsTestSuiteModule>("StructUtilsTestSuite");
+	}
+
+	/**
+	* Checks to see if this module is loaded and ready.  It is only valid to call Get() if IsAvailable() returns true.
+	*
+	* @return True if the module is loaded and ready to use
+	*/
+	static inline bool IsAvailable()
+	{
+		return FModuleManager::Get().IsModuleLoaded("StructUtilsTestSuite");
+	}
+};
+
diff --git a/Engine/Plugins/Experimental/StructUtils/Source/StructUtilsTestSuite/Public/StructUtilsTestTypes.h b/Engine/Plugins/Experimental/StructUtils/Source/StructUtilsTestSuite/Public/StructUtilsTestTypes.h
new file mode 100644
--- /dev/null
+++ b/Engine/Plugins/Experimental/StructUtils/Source/StructUtilsTestSuite/Public/StructUtilsTestTypes.h
@@ -0,0 +1,111 @@
+// Copyright Epic Games, Inc. All Rights Reserved.
+
+#pragma once
+
+#include "CoreMinimal.h"
+#include "StructUtilsTestTypes.generated.h"
+
+USTRUCT()
+struct FTestStructSimple
+{
+	GENERATED_BODY()
+	
+	FTestStructSimple() = default;
+	FTestStructSimple(const float InFloat) : Float(InFloat) {}
+	
+	UPROPERTY()
+	float Float = 0.0f;
+};
+
+USTRUCT()
+struct FTestStructComplex
+{
+	GENERATED_BODY()
+	
+	FTestStructComplex() = default;
+	FTestStructComplex(const FString& InString) : String(InString) {}
+	
+	UPROPERTY()
+	FString String; 
+
+	UPROPERTY()
+	TArray<FString> StringArray;
+};
+
+
+USTRUCT()
+struct FTestStructSimple1
+{
+	GENERATED_BODY()
+	
+	FTestStructSimple1() = default;
+	
+	UPROPERTY()
+	float Float = 0.0f;
+};
+
+USTRUCT()
+struct FTestStructSimple2
+{
+	GENERATED_BODY()
+	
+	FTestStructSimple2() = default;
+	
+	UPROPERTY()
+	float Float = 0.0f;
+};
+
+USTRUCT()
+struct FTestStructSimple3
+{
+	GENERATED_BODY()
+	
+	FTestStructSimple3() = default;
+	
+	UPROPERTY()
+	float Float = 0.0f;
+};
+
+USTRUCT()
+struct FTestStructSimple4
+{
+	GENERATED_BODY()
+	
+	FTestStructSimple4() = default;
+	
+	UPROPERTY()
+	float Float = 0.0f;
+};
+
+USTRUCT()
+struct FTestStructSimple5
+{
+	GENERATED_BODY()
+	
+	FTestStructSimple5() = default;
+	
+	UPROPERTY()
+	float Float = 0.0f;
+};
+
+USTRUCT()
+struct FTestStructSimple6
+{
+	GENERATED_BODY()
+	
+	FTestStructSimple6() = default;
+	
+	UPROPERTY()
+	float Float = 0.0f;
+};
+
+USTRUCT()
+struct FTestStructSimple7
+{
+	GENERATED_BODY()
+	
+	FTestStructSimple7() = default;
+	
+	UPROPERTY()
+	float Float = 0.0f;
+};
diff --git a/Engine/Plugins/Experimental/StructUtils/Source/StructUtilsTestSuite/StructUtilsTestSuite.Build.cs b/Engine/Plugins/Experimental/StructUtils/Source/StructUtilsTestSuite/StructUtilsTestSuite.Build.cs
new file mode 100644
--- /dev/null
+++ b/Engine/Plugins/Experimental/StructUtils/Source/StructUtilsTestSuite/StructUtilsTestSuite.Build.cs
@@ -0,0 +1,35 @@
+// Copyright Epic Games, Inc. All Rights Reserved.
+
+using UnrealBuildTool;
+
+namespace UnrealBuildTool.Rules
+{
+	public class StructUtilsTestSuite : ModuleRules
+	{
+		public StructUtilsTestSuite(ReadOnlyTargetRules Target) : base(Target)
+		{
+			PCHUsage = ModuleRules.PCHUsageMode.UseExplicitOrSharedPCHs;
+
+			PublicIncludePaths.AddRange(
+				new string[] {
+				}
+			);
+
+			PublicDependencyModuleNames.AddRange(
+				new string[] {
+					"Core",
+					"CoreUObject",
+					"Engine",
+					"AIModule",
+					"AITestSuite",
+					"StructUtils"
+				}
+			);
+
+			if (Target.bBuildEditor == true)
+			{
+				PrivateDependencyModuleNames.Add("UnrealEd");
+			}
+		}
+	}
+}
\ No newline at end of file
diff --git a/Engine/Plugins/Experimental/StructUtils/StructUtils.uplugin b/Engine/Plugins/Experimental/StructUtils/StructUtils.uplugin
--- a/Engine/Plugins/Experimental/StructUtils/StructUtils.uplugin
+++ b/Engine/Plugins/Experimental/StructUtils/StructUtils.uplugin
@@ -24,6 +24,10 @@
 		{
 			"Name" : "StructUtilsEditor",
 			"Type" : "Editor"
+		},
+		{
+			"Name": "StructUtilsTestSuite",
+			"Type": "UncookedOnly"
 		}
 	],
 	"Plugins": [
