Added GetAllocationSize() on the FInstanceStructStream

#rnx

diff --git a/Engine/Plugins/Experimental/StructUtils/Source/StructUtils/Public/InstancedStructStream.h b/Engine/Plugins/Experimental/StructUtils/Source/StructUtils/Public/InstancedStructStream.h
--- a/Engine/Plugins/Experimental/StructUtils/Source/StructUtils/Public/InstancedStructStream.h
+++ b/Engine/Plugins/Experimental/StructUtils/Source/StructUtils/Public/InstancedStructStream.h
@@ -1,4 +1,4 @@
-﻿// Copyright Epic Games, Inc. All Rights Reserved.
+// Copyright Epic Games, Inc. All Rights Reserved.
 
 #pragma once
 
@@ -256,6 +256,18 @@ struct STRUCTUTILS_API FInstancedStructStream
 	/** @return True if the buffer is empty. */
 	bool IsEmpty() const { return NumItems == 0; }
 
+	/** 
+	 * Helper function to return the amount of memory allocated by this
+	 * container.
+	 * Only returns the size of allocations made directly by the container, not the elements themselves.
+	 *
+	 * @returns Number of bytes allocated by this container.
+	 */
+	FORCEINLINE SIZE_T GetAllocatedSize(void) const
+	{
+		return (GetNumUsedChunks() + GetNumFreeChunks()) * GetChunkSize();
+	}
+
 	/** @return Chunk size */
 	int32 GetChunkSize() const { return ChunkSize; }
 
