Fix FInstancedStructStream move constructor

diff --git a/Engine/Plugins/Experimental/StructUtils/Source/StructUtils/Public/InstancedStructStream.h b/Engine/Plugins/Experimental/StructUtils/Source/StructUtils/Public/InstancedStructStream.h
--- a/Engine/Plugins/Experimental/StructUtils/Source/StructUtils/Public/InstancedStructStream.h
+++ b/Engine/Plugins/Experimental/StructUtils/Source/StructUtils/Public/InstancedStructStream.h
@@ -35,7 +35,7 @@ struct STRUCTUTILS_API FInstancedStructStream
 	FInstancedStructStream(FInstancedStructStream&& Other)
 	{
 		ChunkSize = Other.ChunkSize;
-		Append(Other);
+		Append(MoveTemp(Other));
 	}
 
 	~FInstancedStructStream()
