Introduction of a generic container type for tracking sub-types of a given UStruct and storing their presence in a bitfield

diff --git a/Engine/Plugins/Experimental/StructUtils/Source/StructUtils/Public/ScriptStructTypeBitSet.h b/Engine/Plugins/Experimental/StructUtils/Source/StructUtils/Public/ScriptStructTypeBitSet.h
new file mode 100644
--- /dev/null
+++ b/Engine/Plugins/Experimental/StructUtils/Source/StructUtils/Public/ScriptStructTypeBitSet.h
@@ -0,0 +1,446 @@
+// Copyright Epic Games, Inc. All Rights Reserved.
+
+#pragma once
+
+#include "UObject/Class.h"
+#include "Containers/BitArray.h"
+#include "InstancedStruct.h"
+#include "StructUtilsTypes.h"
+
+/**
+ * The TScriptStructTypeBitSet holds information on "existence" of subtypes of a given UStruct. The information on 
+ * available child-structs is gathered lazily - the internal FStructTracker is assigning a given type a new index the
+ * very first time the type is encountered. 
+ * To create a specific instantiation of the type you need to declare the static StructTracker member variable. 
+ * To do it for an arbitrary type FFooBar add the following in your header or cpp file:
+ * 
+ *	DECLARE_STRUCTTYPEBITSET(FMyFooBarBitSet, FFooBar);
+ * 
+ * where FMyFooBarBitSet is the alias of the type you can use in your code. To have your type exposed to other modules 
+ * use DECLARE_STRUCTTYPEBITSET_EXPORTED, like so:
+ * 
+ *	DECLARE_STRUCTTYPEBITSET_EXPORTED(MYMODULE_API, FMyFooBarBitSet, FFooBar);
+ *
+ * The type contains static members so you'll also need to define these. You can easily do it by placing the following 
+ * in your cpp file (continuing the FFooBar example):
+ *
+ *	DEFINE_STRUCTTYPEBITSET(FMyFooBarBitSet);
+ * 
+ */
+template<typename TBaseStruct>
+struct TScriptStructTypeBitSet
+{
+private:
+	struct FBitArrayExt : TBitArray<>
+	{
+		FBitArrayExt() = default;
+		FBitArrayExt(const TBitArray<>& Source) : TBitArray<>(Source)
+		{}
+
+		FBitArrayExt& operator=(const TBitArray<>& Other)
+		{
+			*((TBitArray<>*)this) = Other;
+			return *this;
+		}
+
+		FORCEINLINE bool HasAll(const TBitArray<>& Other) const
+		{
+			FConstWordIterator ThisIterator(*this);
+			FConstWordIterator OtherIterator(Other);
+
+			while (ThisIterator || OtherIterator)
+			{
+				const uint32 A = ThisIterator ? ThisIterator.GetWord() : 0;
+				const uint32 B = OtherIterator ? OtherIterator.GetWord() : 0;
+				if ((A & B) != B)
+				{
+					return false;
+				}
+
+				++ThisIterator;
+				++OtherIterator;
+			}
+
+			return true;
+		}
+
+		FORCEINLINE bool HasAny(const TBitArray<>& Other) const
+		{
+			FConstWordIterator ThisIterator(*this);
+			FConstWordIterator OtherIterator(Other);
+
+			while (ThisIterator || OtherIterator)
+			{
+				const uint32 A = ThisIterator ? ThisIterator.GetWord() : 0;
+				const uint32 B = OtherIterator ? OtherIterator.GetWord() : 0;
+				if ((A & B) != 0)
+				{
+					return true;
+				}
+
+				++ThisIterator;
+				++OtherIterator;
+			}
+
+			return false;
+		}
+
+		FORCEINLINE bool IsEmpty() const
+		{
+			FConstWordIterator Iterator(*this);
+
+			while (Iterator && Iterator.GetWord() == 0)
+			{
+				++Iterator;
+			}
+
+			return !Iterator;
+		}
+
+		FORCEINLINE void operator-=(const TBitArray<>& Other)
+		{
+			FWordIterator ThisIterator(*this);
+			FConstWordIterator OtherIterator(Other);
+
+			while (ThisIterator && OtherIterator)
+			{
+				ThisIterator.SetWord(ThisIterator.GetWord() & ~OtherIterator.GetWord());
+
+				++ThisIterator;
+				++OtherIterator;
+			}
+		}
+
+		FORCEINLINE friend uint32 GetTypeHash(const FBitArrayExt& Instance)
+		{
+			FConstWordIterator Iterator(Instance);
+			uint32 Hash = 0;
+			while (Iterator)
+			{
+				Hash = HashCombine(Hash, Iterator.GetWord());
+				++Iterator;
+			}
+			return Hash;
+		}
+	};
+
+	struct FStructTracker
+	{
+		int32 FindOrAddStructTypeIndex(const UScriptStruct& InStructType)
+		{
+#if WITH_STRUCTUTILS_DEBUG
+			ensureMsgf(InStructType.IsChildOf(TBaseStruct::StaticStruct())
+				, TEXT("Registering '%s' with FStructTracker while it doesn't derive from the expected struct type %s")
+				, *InStructType.GetPathName(), *TBaseStruct::StaticStruct()->GetName());
+#endif // WITH_STRUCTUTILS_DEBUG
+
+			// Get existing index...
+			const uint32 Hash = PointerHash(&InStructType);
+			FSetElementId ElementId = StructTypeToIndexSet.FindIdByHash(Hash, Hash);
+
+			if (!ElementId.IsValidId())
+			{
+				// .. or create new one
+				ElementId = StructTypeToIndexSet.AddByHash(Hash, Hash);
+				checkSlow(ElementId.IsValidId());
+
+				checkSlow(StructTypesList.Num() == ElementId.AsInteger());
+				StructTypesList.Add(&InStructType);
+			}
+			const int32 Index = ElementId.AsInteger();
+			
+#if WITH_STRUCTUTILS_DEBUG
+			if (Index == DebugStructTypeNamesList.Num())
+			{
+				DebugStructTypeNamesList.Add(InStructType.GetFName());
+				ensure(StructTypeToIndexSet.Num() == DebugStructTypeNamesList.Num());
+			}
+#endif // WITH_STRUCTUTILS_DEBUG
+			return Index;
+		}
+
+		const UScriptStruct* GetStructType(const int32 StructTypeIndex) const
+		{
+			return StructTypesList.IsValidIndex(StructTypeIndex) ? StructTypesList[StructTypeIndex].Get() : nullptr;
+		}
+
+#if WITH_STRUCTUTILS_DEBUG
+		/**
+		 * @return index identifying given tag or INDEX_NONE if it has never been used/seen before.
+		 */
+		FName DebugGetStructTypeName(const int32 StructTypeIndex) const
+		{
+			return DebugStructTypeNamesList.IsValidIndex(StructTypeIndex) ? DebugStructTypeNamesList[StructTypeIndex] : FName();
+		}
+
+		void DebugNukeStructTypeMappingInfo()
+		{
+			StructTypeToIndexSet.Reset();
+			StructTypesList.Reset();
+			DebugStructTypeNamesList.Reset();
+		}
+		TArray<FName, TInlineAllocator<64>> DebugStructTypeNamesList;
+#endif // WITH_STRUCTUTILS_DEBUG
+
+		TSet<uint32> StructTypeToIndexSet;
+		TArray<TWeakObjectPtr<const UScriptStruct>, TInlineAllocator<64>> StructTypesList;
+	};
+	static FStructTracker StructTracker;
+
+public:
+	TScriptStructTypeBitSet() = default;
+
+	explicit TScriptStructTypeBitSet(const UScriptStruct& StructType)
+	{
+		Add(StructType);
+	}
+
+	explicit TScriptStructTypeBitSet(std::initializer_list<const UScriptStruct*> InitList)
+	{
+		for (const UScriptStruct* StructType : InitList)
+		{
+			if (StructType)
+			{
+				Add(*StructType);
+			}
+		}
+	}
+
+	explicit TScriptStructTypeBitSet(TConstArrayView<const UScriptStruct*> InitList)
+	{
+		for (const UScriptStruct* StructType : InitList)
+		{
+			if (StructType)
+			{
+				Add(*StructType);
+			}
+		}
+	}
+
+	explicit TScriptStructTypeBitSet(TConstArrayView<FInstancedStruct> InitList)
+	{
+		for (const FInstancedStruct& StructInstance : InitList)
+		{
+			if (StructInstance.GetScriptStruct())
+			{
+				Add(*StructInstance.GetScriptStruct());
+			}
+		}
+	}
+
+private:
+	/** 
+	 * A private constructor for a creating an instance straight from TBitArrays. 
+	 * @Note that this constructor needs to remain private to ensure consistency of stored values with data tracked 
+	 * by the StructTracker
+	 */
+	TScriptStructTypeBitSet(const TBitArray<>& Source)
+		: StructTypesBitArray(Source)
+	{
+	}
+
+public:
+
+	template<typename T>
+	FORCEINLINE void Add()
+	{
+		static_assert(TIsDerivedFrom<T, TBaseStruct>::IsDerived, "Given struct type doesn't match the expected base struct type.");
+		Add(*T::StaticStruct());
+	}
+
+	template<typename T>
+	FORCEINLINE void Remove()
+	{
+		static_assert(TIsDerivedFrom<T, TBaseStruct>::IsDerived, "Given struct type doesn't match the expected base struct type.");
+		Remove(*T::StaticStruct());
+	}
+
+	FORCEINLINE void Remove(const TScriptStructTypeBitSet<TBaseStruct>& Other)
+	{
+		StructTypesBitArray -= Other.StructTypesBitArray;
+	}
+
+	template<typename T>
+	FORCEINLINE bool Contains() const
+	{
+		static_assert(TIsDerivedFrom<T, TBaseStruct>::IsDerived, "Given struct type doesn't match the expected base struct type.");
+		return Contains(*T::StaticStruct());
+	}
+
+	void Add(const UScriptStruct& InStructType)
+	{
+		const int32 StructTypeIndex = StructTracker.FindOrAddStructTypeIndex(InStructType);
+		StructTypesBitArray.PadToNum(StructTypeIndex + 1, false);
+		StructTypesBitArray[StructTypeIndex] = true;
+	}
+
+	void Remove(const UScriptStruct& InStructType)
+	{
+		const int32 StructTypeIndex = StructTracker.FindOrAddStructTypeIndex(InStructType);
+		if (StructTypeIndex < StructTypesBitArray.Num())
+		{
+			StructTypesBitArray[StructTypeIndex] = false;
+		}
+		// else, if it's not in StructTypesBitArray it's already not present
+	}
+
+	void Reset() { StructTypesBitArray.Reset(); }
+
+	bool Contains(const UScriptStruct& InStructType) const
+	{
+		const int32 StructTypeIndex = StructTracker.FindOrAddStructTypeIndex(InStructType);
+		return (StructTypeIndex < StructTypesBitArray.Num()) && StructTypesBitArray[StructTypeIndex];
+	}
+
+	FORCEINLINE TScriptStructTypeBitSet operator+(const TScriptStructTypeBitSet& Other) const
+	{
+		TScriptStructTypeBitSet Result;
+		Result.StructTypesBitArray = TBitArray<>::BitwiseOR(StructTypesBitArray, Other.StructTypesBitArray, EBitwiseOperatorFlags::MaxSize);
+		return MoveTemp(Result);
+	}
+
+	FORCEINLINE void operator+=(const TScriptStructTypeBitSet& Other)
+	{
+		StructTypesBitArray = TBitArray<>::BitwiseOR(StructTypesBitArray, Other.StructTypesBitArray, EBitwiseOperatorFlags::MaxSize);
+	}
+
+	FORCEINLINE void operator-=(const TScriptStructTypeBitSet& Other)
+	{
+		StructTypesBitArray -= Other.StructTypesBitArray;
+	}
+
+	FORCEINLINE TScriptStructTypeBitSet operator+(const UScriptStruct& NewElement) const
+	{
+		TScriptStructTypeBitSet Result = *this;
+		Result.Add(NewElement);
+		return MoveTemp(Result);
+	}
+
+	FORCEINLINE TScriptStructTypeBitSet operator-(const UScriptStruct& NewElement) const
+	{
+		TScriptStructTypeBitSet Result = *this;
+		Result.Remove(NewElement);
+		return MoveTemp(Result);
+	}
+
+	FORCEINLINE TScriptStructTypeBitSet operator-(const TScriptStructTypeBitSet& Other) const
+	{
+		TScriptStructTypeBitSet Result = *this;
+		Result -= Other;
+		return MoveTemp(Result);
+	}
+
+	FORCEINLINE TScriptStructTypeBitSet operator&(const TScriptStructTypeBitSet& Other) const
+	{
+		return TScriptStructTypeBitSet(TBitArray<>::BitwiseAND(StructTypesBitArray, Other.StructTypesBitArray, EBitwiseOperatorFlags::MinSize));
+	}
+
+	FORCEINLINE TScriptStructTypeBitSet GetOverlap(const TScriptStructTypeBitSet& Other) const
+	{
+		return *this & Other;
+	}
+
+	FORCEINLINE bool IsEquivalent(const TScriptStructTypeBitSet<TBaseStruct>& Other) const
+	{
+		return StructTypesBitArray.CompareSetBits(Other.StructTypesBitArray, /*bMissingBitValue=*/false);
+	}
+
+	FORCEINLINE bool HasAll(const TScriptStructTypeBitSet& Other) const
+	{
+		return StructTypesBitArray.HasAll(Other.StructTypesBitArray);
+	}
+
+	FORCEINLINE bool HasAny(const TScriptStructTypeBitSet& Other) const
+	{
+		return StructTypesBitArray.HasAny(Other.StructTypesBitArray);
+	}
+
+	FORCEINLINE bool HasNone(const TScriptStructTypeBitSet& Other) const
+	{
+		return !StructTypesBitArray.HasAny(Other.StructTypesBitArray);
+	}
+
+	bool IsEmpty() const 
+	{ 
+		return StructTypesBitArray.IsEmpty();
+	}
+
+	FORCEINLINE bool operator==(const TScriptStructTypeBitSet& Other) const { return StructTypesBitArray == Other.StructTypesBitArray; }
+	FORCEINLINE bool operator!=(const TScriptStructTypeBitSet& Other) const { return !(*this == Other); }
+
+	/**
+	 * note that this function is slow(ish) due to the FStructTracker utilizing WeakObjectPtrs to store types. 
+	 * @todo To be improved.
+	 */
+	template<typename Allocator>
+	void ExportTypes(TArray<const UScriptStruct*, Allocator>& OutTypes) const
+	{
+		TBitArray<>::FConstIterator It(StructTypesBitArray);
+		while (It)
+		{
+			if (It.GetValue())
+			{
+				OutTypes.Add(StructTracker.GetStructType(It.GetIndex()));
+			}
+			++It;
+		}
+	}
+
+#if WITH_STRUCTUTILS_DEBUG
+	void DebugGetStringDesc(FOutputDevice& Ar) const
+	{
+		for (int32 Index = 0; Index < StructTypesBitArray.Num(); ++Index)
+		{
+			if (StructTypesBitArray[Index])
+			{
+				Ar.Logf(TEXT("%s, "), *StructTracker.DebugGetStructTypeName(Index).ToString());
+			}
+		}
+	}
+
+	void DebugGetIndividualNames(TArray<FName>& OutFNames) const
+	{
+		for (int32 Index = 0; Index < StructTypesBitArray.Num(); ++Index)
+		{
+			if (StructTypesBitArray[Index])
+			{
+				OutFNames.Add(StructTracker.DebugGetStructTypeName(Index));
+			}
+		}
+	}
+
+	/**
+	 * Resets all the information gathered on the tags. Calling this results in invalidating all previously created
+	 * FStructTypeCollection instances. Used only for debugging and unit/functional testing.
+	 */
+	static void DebugNukeStructTypeMappingInfo()
+	{
+		StructTracker.DebugNukeStructTypeMappingInfo();
+	}
+protected:
+	// unittesting purposes only
+	const TBitArray<>& DebugGetStructTypesBitArray() const { return StructTypesBitArray; }
+	TBitArray<>& DebugGetMutableStructTypesBitArray() { return StructTypesBitArray; }
+#endif // WITH_STRUCTUTILS_DEBUG
+
+public:
+	FORCEINLINE friend uint32 GetTypeHash(const TScriptStructTypeBitSet<TBaseStruct>& Instance)
+	{
+		const uint32 BitArrayHash = GetTypeHash(Instance.StructTypesBitArray);
+		const uint32 StoredTypeHash = PointerHash(TBaseStruct::StaticStruct());
+		return HashCombine(StoredTypeHash, BitArrayHash);
+	}
+
+private:
+	FBitArrayExt StructTypesBitArray;
+};
+
+#define DECLARE_STRUCTTYPEBITSET_EXPORTED(EXPORTED_API, ContainerTypeName, BaseStructType) template<> \
+	TScriptStructTypeBitSet<BaseStructType>::FStructTracker TScriptStructTypeBitSet<BaseStructType>::StructTracker; \
+	template struct EXPORTED_API TScriptStructTypeBitSet<BaseStructType>; \
+	using ContainerTypeName = TScriptStructTypeBitSet<BaseStructType>;
+
+#define DECLARE_STRUCTTYPEBITSET(ContainerTypeName, BaseStructType) DECLARE_STRUCTTYPEBITSET_EXPORTED(, ContainerTypeName, BaseStructType)
+
+#define DEFINE_STRUCTTYPEBITSET(ContainerTypeName) template<> \
+	ContainerTypeName::FStructTracker ContainerTypeName::StructTracker = ContainerTypeName::FStructTracker();
diff --git a/Engine/Plugins/Experimental/StructUtils/Source/StructUtils/Public/StructUtilsTypes.h b/Engine/Plugins/Experimental/StructUtils/Source/StructUtils/Public/StructUtilsTypes.h
new file mode 100644
--- /dev/null
+++ b/Engine/Plugins/Experimental/StructUtils/Source/StructUtils/Public/StructUtilsTypes.h
@@ -0,0 +1,7 @@
+// Copyright Epic Games, Inc. All Rights Reserved.
+
+#pragma once
+
+#ifndef WITH_STRUCTUTILS_DEBUG
+#define WITH_STRUCTUTILS_DEBUG (!(UE_BUILD_SHIPPING || UE_BUILD_SHIPPING_WITH_EDITOR || UE_BUILD_TEST) && 1)
+#endif // WITH_STRUCTUTILS_DEBUG
\ No newline at end of file
diff --git a/Engine/Plugins/Experimental/StructUtils/Source/StructUtilsTestSuite/Private/ScriptStructTypeBitSetTest.cpp b/Engine/Plugins/Experimental/StructUtils/Source/StructUtilsTestSuite/Private/ScriptStructTypeBitSetTest.cpp
new file mode 100644
--- /dev/null
+++ b/Engine/Plugins/Experimental/StructUtils/Source/StructUtilsTestSuite/Private/ScriptStructTypeBitSetTest.cpp
@@ -0,0 +1,246 @@
+// Copyright Epic Games, Inc. All Rights Reserved.
+
+#include "CoreMinimal.h"
+#include "AITestsCommon.h"
+#include "ScriptStructTypeBitSet.h"
+#include "StructUtilsTestTypes.h"
+
+#if WITH_STRUCTUTILS_DEBUG
+
+#define LOCTEXT_NAMESPACE "StructUtilsTests"
+
+PRAGMA_DISABLE_OPTIMIZATION
+
+template<>
+TScriptStructTypeBitSet<FTestStructSimple>::FStructTracker TScriptStructTypeBitSet<FTestStructSimple>::StructTracker = TScriptStructTypeBitSet<FTestStructSimple>::FStructTracker();
+template struct TScriptStructTypeBitSet<FTestStructSimple>::FStructTracker;
+
+namespace FScriptStructTypeBitSetTests
+{
+
+struct FTestStructBitSet : public TScriptStructTypeBitSet<FTestStructSimple>
+{
+	FTestStructBitSet() = default;
+	FTestStructBitSet(const TScriptStructTypeBitSet<FTestStructSimple>& Other) 
+		: TScriptStructTypeBitSet<FTestStructSimple>(Other)
+	{
+
+	}
+
+	void AddBit(const int32 Index)
+	{
+		DebugGetMutableStructTypesBitArray().PadToNum(Index + 1, false);
+		DebugGetMutableStructTypesBitArray()[Index] = true;
+	}
+
+	void RemoveBit(const int32 Index)
+	{
+		DebugGetMutableStructTypesBitArray()[Index] = false;
+	}
+
+	bool TestBit(const int32 Index) const
+	{
+		return Index >= 0 && DebugGetStructTypesBitArray()[Index];
+	}
+};
+
+struct FStructUtilsTest_BitSetEquivalence : FAITestBase
+{
+	virtual bool InstantTest() override
+	{
+		FTestStructBitSet CollectionA;
+		FTestStructBitSet CollectionB;
+
+		AITEST_TRUE("Empty collections are equivalent", CollectionA.IsEquivalent(CollectionB));
+		AITEST_TRUE("Equivalence check is commutative", CollectionA.IsEquivalent(CollectionB) == CollectionB.IsEquivalent(CollectionA));
+
+		CollectionA.AddBit(1);
+		AITEST_FALSE("Given collections are not equivalent", CollectionA.IsEquivalent(CollectionB));
+		AITEST_TRUE("Equivalence check is commutative", CollectionA.IsEquivalent(CollectionB) == CollectionB.IsEquivalent(CollectionA));
+
+		CollectionB.AddBit(1);
+		AITEST_TRUE("Given collections are equivalent", CollectionA.IsEquivalent(CollectionB));
+		AITEST_TRUE("Equivalence check is commutative", CollectionA.IsEquivalent(CollectionB) == CollectionB.IsEquivalent(CollectionA));
+
+		CollectionA.AddBit(124);
+		AITEST_FALSE("Given collections are not equivalent", CollectionA.IsEquivalent(CollectionB));
+		AITEST_TRUE("Equivalence check is commutative", CollectionA.IsEquivalent(CollectionB) == CollectionB.IsEquivalent(CollectionA));
+
+		CollectionA.RemoveBit(124);
+		AITEST_TRUE("Given collections are equivalent", CollectionA.IsEquivalent(CollectionB));
+		AITEST_TRUE("Equivalence check is commutative", CollectionA.IsEquivalent(CollectionB) == CollectionB.IsEquivalent(CollectionA));
+
+		return true;
+	}
+};
+
+IMPLEMENT_AI_INSTANT_TEST(FStructUtilsTest_BitSetEquivalence, "System.StructUtils.BitSet.Equivalence");
+
+struct FStructUtilsTest_BitSetEmptiness : FAITestBase
+{
+	virtual bool InstantTest() override
+	{
+		FTestStructBitSet Collection;
+
+		AITEST_TRUE("New collection is empty", Collection.IsEmpty());
+				
+		Collection.AddBit(125);
+		AITEST_FALSE("Extended collection is not empty", Collection.IsEmpty());
+		
+		Collection.RemoveBit(125);
+		AITEST_TRUE("Removing the remote bit should make the collection empty again", Collection.IsEmpty());
+
+		return true;
+	}
+};
+
+IMPLEMENT_AI_INSTANT_TEST(FStructUtilsTest_BitSetEmptiness, "System.StructUtils.BitSet.Emptiness");
+
+struct FStructUtilsTest_BitSetComparison : FAITestBase
+{
+	virtual bool InstantTest() override
+	{
+		FTestStructBitSet CollectionA;
+		FTestStructBitSet CollectionNone;
+		FTestStructBitSet CollectionAll;
+		FTestStructBitSet CollectionSome;
+
+		CollectionA.AddBit(1);
+		CollectionA.AddBit(32);
+		CollectionSome = CollectionA;
+		CollectionSome.AddBit(111);
+
+		CollectionA.AddBit(65);
+		CollectionAll = CollectionA;
+
+		CollectionA.AddBit(76);
+
+		CollectionNone.AddBit(2);
+		CollectionNone.AddBit(77);
+
+		AITEST_TRUE("Given collection should confirm it has all its element", CollectionA.HasAll(CollectionA));
+		AITEST_TRUE("CollectionA has all the elements indicated by CollectionAll set", CollectionA.HasAll(CollectionAll));
+		AITEST_FALSE("HasAll is not commutative", CollectionAll.HasAll(CollectionA));
+
+		AITEST_TRUE("CollectionA has none of the elements stored in CollectionNone", CollectionA.HasNone(CollectionNone));
+		AITEST_TRUE("HasNone is commutative", CollectionNone.HasNone(CollectionA));
+
+		AITEST_TRUE("", CollectionA.HasAny(CollectionSome));
+		AITEST_TRUE("HasAny is commutative", CollectionSome.HasAny(CollectionA));
+
+		return true;
+	}
+};
+
+IMPLEMENT_AI_INSTANT_TEST(FStructUtilsTest_BitSetComparison, "System.StructUtils.BitSet.Comparison");
+
+struct FStructUtilsTest_BitSetSubtraction : FAITestBase
+{
+	virtual bool InstantTest() override
+	{
+		constexpr int TotalBits = 60;
+		constexpr int BitsToClear = 40;
+
+		FTestStructBitSet CollectionA;
+		FTestStructBitSet CollectionB;
+		
+		for (int i = 0; i < TotalBits; ++i)
+		{
+			CollectionA.AddBit(i);
+		}
+		for (int i = 0; i < BitsToClear; ++i)
+		{
+			CollectionB.AddBit(i);
+		}
+		
+		FTestStructBitSet CollectionC = CollectionA - CollectionB;
+
+		for (int i = 0; i < BitsToClear; ++i)
+		{
+			AITEST_FALSE("Testing expected bit cleared", CollectionC.TestBit(i));
+		}
+		for (int i = BitsToClear; i < TotalBits; ++i)
+		{
+			AITEST_TRUE("Testing expected bit remaining", CollectionC.TestBit(i));
+		}
+
+		return true;
+	}
+};
+
+IMPLEMENT_AI_INSTANT_TEST(FStructUtilsTest_BitSetSubtraction, "System.StructUtils.BitSet.Subtraction");
+
+struct FStructUtilsTest_BitSetOverlap : FAITestBase
+{
+	virtual bool InstantTest() override
+	{
+		constexpr int BitsACount = 40;
+		constexpr int OverlapBitsCount = 10;
+		constexpr int BitsBCount = 30;
+
+		FTestStructBitSet CollectionA;
+		FTestStructBitSet CollectionB;
+
+		for (int i = 0; i < BitsACount; ++i)
+		{
+			CollectionA.AddBit(i);
+		}
+		for (int i = 0; i < BitsBCount; ++i)
+		{
+			CollectionB.AddBit(BitsACount - OverlapBitsCount + i);
+		}
+
+		FTestStructBitSet CollectionC = CollectionA & CollectionB;
+		FTestStructBitSet CollectionD = CollectionB & CollectionA;
+
+		AITEST_TRUE("Overlap operator is commutative", CollectionC.IsEquivalent(CollectionD));
+
+		int j = 0;
+		for (int i = 0; i < BitsACount - OverlapBitsCount; ++i, ++j)
+		{
+			AITEST_FALSE("Testing not-overlapping bits", CollectionC.TestBit(j));
+		}
+		for (int i = 0; i < OverlapBitsCount; ++i, ++j)
+		{
+			AITEST_TRUE("Testing overlapping bits", CollectionC.TestBit(j));
+		}
+		for (int i = 0; i < BitsBCount - OverlapBitsCount; ++i, ++j)
+		{
+			AITEST_FALSE("Testing remaining non-overlapping bits", CollectionC.TestBit(j));
+		}
+
+		return true;
+	}
+};
+
+IMPLEMENT_AI_INSTANT_TEST(FStructUtilsTest_BitSetOverlap, "System.StructUtils.BitSet.Overlap");
+
+struct FStructUtilsTest_BitSetHash : FAITestBase
+{
+	virtual bool InstantTest() override
+	{
+		FTestStructBitSet CollectionA;
+		FTestStructBitSet CollectionB;
+
+		CollectionA.AddBit(9);
+		CollectionB.AddBit(9);
+		CollectionB.AddBit(21);
+
+		const uint32 HashA = GetTypeHash(CollectionA);
+		const uint32 HashB = GetTypeHash(CollectionB);
+
+		AITEST_NOT_EQUAL("Two distinct bit sets should have distinct hashes", HashA, HashB);
+
+		return true;
+	}
+};
+
+IMPLEMENT_AI_INSTANT_TEST(FStructUtilsTest_BitSetHash, "System.StructUtils.BitSet.Hash");
+
+} // namespace FScriptStructTypeBitSetTests
+
+#undef LOCTEXT_NAMESPACE
+
+PRAGMA_ENABLE_OPTIMIZATION
+
+#endif // WITH_STRUCTUTILS_DEBUG
