ScriptStructTypeBitSet unit test fix

diff --git a/Engine/Plugins/Experimental/StructUtils/Source/StructUtilsTestSuite/Private/ScriptStructTypeBitSetTest.cpp b/Engine/Plugins/Experimental/StructUtils/Source/StructUtilsTestSuite/Private/ScriptStructTypeBitSetTest.cpp
--- a/Engine/Plugins/Experimental/StructUtils/Source/StructUtilsTestSuite/Private/ScriptStructTypeBitSetTest.cpp
+++ b/Engine/Plugins/Experimental/StructUtils/Source/StructUtilsTestSuite/Private/ScriptStructTypeBitSetTest.cpp
@@ -40,7 +40,7 @@ struct FTestStructBitSet : public TScriptStructTypeBitSet<FTestStructSimple>
 
 	bool TestBit(const int32 Index) const
 	{
-		return Index >= 0 && DebugGetStructTypesBitArray()[Index];
+		return Index >= 0 && Index < DebugGetStructTypesBitArray().Num() && DebugGetStructTypesBitArray()[Index];
 	}
 };
 
