Fix ClangEditor Win64

diff --git a/Engine/Plugins/Experimental/StructUtils/Source/StructUtils/Public/ScriptStructTypeBitSet.h b/Engine/Plugins/Experimental/StructUtils/Source/StructUtils/Public/ScriptStructTypeBitSet.h
--- a/Engine/Plugins/Experimental/StructUtils/Source/StructUtils/Public/ScriptStructTypeBitSet.h
+++ b/Engine/Plugins/Experimental/StructUtils/Source/StructUtils/Public/ScriptStructTypeBitSet.h
@@ -27,6 +27,62 @@
  *	DEFINE_STRUCTTYPEBITSET(FMyFooBarBitSet);
  * 
  */
+struct FStructTracker
+{
+	int32 FindOrAddStructTypeIndex(const UScriptStruct& InStructType)
+	{
+		// Get existing index...
+		const uint32 Hash = PointerHash(&InStructType);
+		FSetElementId ElementId = StructTypeToIndexSet.FindIdByHash(Hash, Hash);
+
+		if (!ElementId.IsValidId())
+		{
+			// .. or create new one
+			ElementId = StructTypeToIndexSet.AddByHash(Hash, Hash);
+			checkSlow(ElementId.IsValidId());
+
+			checkSlow(StructTypesList.Num() == ElementId.AsInteger());
+			StructTypesList.Add(&InStructType);
+		}
+		const int32 Index = ElementId.AsInteger();
+
+#if WITH_STRUCTUTILS_DEBUG
+		if (Index == DebugStructTypeNamesList.Num())
+		{
+			DebugStructTypeNamesList.Add(InStructType.GetFName());
+			ensure(StructTypeToIndexSet.Num() == DebugStructTypeNamesList.Num());
+		}
+#endif // WITH_STRUCTUTILS_DEBUG
+		return Index;
+	}
+
+	const UScriptStruct* GetStructType(const int32 StructTypeIndex) const
+	{
+		return StructTypesList.IsValidIndex(StructTypeIndex) ? StructTypesList[StructTypeIndex].Get() : nullptr;
+	}
+
+#if WITH_STRUCTUTILS_DEBUG
+	/**
+	* @return index identifying given tag or INDEX_NONE if it has never been used/seen before.
+	*/
+	FName DebugGetStructTypeName(const int32 StructTypeIndex) const
+	{
+		return DebugStructTypeNamesList.IsValidIndex(StructTypeIndex) ? DebugStructTypeNamesList[StructTypeIndex] : FName();
+	}
+
+	void DebugNukeStructTypeMappingInfo()
+	{
+		StructTypeToIndexSet.Reset();
+		StructTypesList.Reset();
+		DebugStructTypeNamesList.Reset();
+	}
+	TArray<FName, TInlineAllocator<64>> DebugStructTypeNamesList;
+#endif // WITH_STRUCTUTILS_DEBUG
+
+	TSet<uint32> StructTypeToIndexSet;
+	TArray<TWeakObjectPtr<const UScriptStruct>, TInlineAllocator<64>> StructTypesList;
+};
+
 template<typename TBaseStruct>
 struct TScriptStructTypeBitSet
 {
@@ -124,67 +180,6 @@ private:
 		}
 	};
 
-	struct FStructTracker
-	{
-		int32 FindOrAddStructTypeIndex(const UScriptStruct& InStructType)
-		{
-#if WITH_STRUCTUTILS_DEBUG
-			ensureMsgf(InStructType.IsChildOf(TBaseStruct::StaticStruct())
-				, TEXT("Registering '%s' with FStructTracker while it doesn't derive from the expected struct type %s")
-				, *InStructType.GetPathName(), *TBaseStruct::StaticStruct()->GetName());
-#endif // WITH_STRUCTUTILS_DEBUG
-
-			// Get existing index...
-			const uint32 Hash = PointerHash(&InStructType);
-			FSetElementId ElementId = StructTypeToIndexSet.FindIdByHash(Hash, Hash);
-
-			if (!ElementId.IsValidId())
-			{
-				// .. or create new one
-				ElementId = StructTypeToIndexSet.AddByHash(Hash, Hash);
-				checkSlow(ElementId.IsValidId());
-
-				checkSlow(StructTypesList.Num() == ElementId.AsInteger());
-				StructTypesList.Add(&InStructType);
-			}
-			const int32 Index = ElementId.AsInteger();
-			
-#if WITH_STRUCTUTILS_DEBUG
-			if (Index == DebugStructTypeNamesList.Num())
-			{
-				DebugStructTypeNamesList.Add(InStructType.GetFName());
-				ensure(StructTypeToIndexSet.Num() == DebugStructTypeNamesList.Num());
-			}
-#endif // WITH_STRUCTUTILS_DEBUG
-			return Index;
-		}
-
-		const UScriptStruct* GetStructType(const int32 StructTypeIndex) const
-		{
-			return StructTypesList.IsValidIndex(StructTypeIndex) ? StructTypesList[StructTypeIndex].Get() : nullptr;
-		}
-
-#if WITH_STRUCTUTILS_DEBUG
-		/**
-		 * @return index identifying given tag or INDEX_NONE if it has never been used/seen before.
-		 */
-		FName DebugGetStructTypeName(const int32 StructTypeIndex) const
-		{
-			return DebugStructTypeNamesList.IsValidIndex(StructTypeIndex) ? DebugStructTypeNamesList[StructTypeIndex] : FName();
-		}
-
-		void DebugNukeStructTypeMappingInfo()
-		{
-			StructTypeToIndexSet.Reset();
-			StructTypesList.Reset();
-			DebugStructTypeNamesList.Reset();
-		}
-		TArray<FName, TInlineAllocator<64>> DebugStructTypeNamesList;
-#endif // WITH_STRUCTUTILS_DEBUG
-
-		TSet<uint32> StructTypeToIndexSet;
-		TArray<TWeakObjectPtr<const UScriptStruct>, TInlineAllocator<64>> StructTypesList;
-	};
 	static FStructTracker StructTracker;
 
 public:
@@ -269,6 +264,12 @@ public:
 
 	void Add(const UScriptStruct& InStructType)
 	{
+#if WITH_STRUCTUTILS_DEBUG
+		ensureMsgf(InStructType.IsChildOf(TBaseStruct::StaticStruct())
+				, TEXT("Registering '%s' with FStructTracker while it doesn't derive from the expected struct type %s")
+				, *InStructType.GetPathName(), *TBaseStruct::StaticStruct()->GetName());
+#endif // WITH_STRUCTUTILS_DEBUG
+
 		const int32 StructTypeIndex = StructTracker.FindOrAddStructTypeIndex(InStructType);
 		StructTypesBitArray.PadToNum(StructTypeIndex + 1, false);
 		StructTypesBitArray[StructTypeIndex] = true;
@@ -276,6 +277,12 @@ public:
 
 	void Remove(const UScriptStruct& InStructType)
 	{
+#if WITH_STRUCTUTILS_DEBUG
+		ensureMsgf(InStructType.IsChildOf(TBaseStruct::StaticStruct())
+				, TEXT("Registering '%s' with FStructTracker while it doesn't derive from the expected struct type %s")
+				, *InStructType.GetPathName(), *TBaseStruct::StaticStruct()->GetName());
+#endif // WITH_STRUCTUTILS_DEBUG
+
 		const int32 StructTypeIndex = StructTracker.FindOrAddStructTypeIndex(InStructType);
 		if (StructTypeIndex < StructTypesBitArray.Num())
 		{
@@ -288,6 +295,12 @@ public:
 
 	bool Contains(const UScriptStruct& InStructType) const
 	{
+#if WITH_STRUCTUTILS_DEBUG
+		ensureMsgf(InStructType.IsChildOf(TBaseStruct::StaticStruct())
+				, TEXT("Registering '%s' with FStructTracker while it doesn't derive from the expected struct type %s")
+				, *InStructType.GetPathName(), *TBaseStruct::StaticStruct()->GetName());
+#endif // WITH_STRUCTUTILS_DEBUG
+
 		const int32 StructTypeIndex = StructTracker.FindOrAddStructTypeIndex(InStructType);
 		return (StructTypeIndex < StructTypesBitArray.Num()) && StructTypesBitArray[StructTypeIndex];
 	}
@@ -435,12 +448,4 @@ private:
 	FBitArrayExt StructTypesBitArray;
 };
 
-#define DECLARE_STRUCTTYPEBITSET_EXPORTED(EXPORTED_API, ContainerTypeName, BaseStructType) template<> \
-	TScriptStructTypeBitSet<BaseStructType>::FStructTracker TScriptStructTypeBitSet<BaseStructType>::StructTracker; \
-	template struct EXPORTED_API TScriptStructTypeBitSet<BaseStructType>; \
-	using ContainerTypeName = TScriptStructTypeBitSet<BaseStructType>;
-
-#define DECLARE_STRUCTTYPEBITSET(ContainerTypeName, BaseStructType) DECLARE_STRUCTTYPEBITSET_EXPORTED(, ContainerTypeName, BaseStructType)
-
-#define DEFINE_STRUCTTYPEBITSET(ContainerTypeName) template<> \
-	ContainerTypeName::FStructTracker ContainerTypeName::StructTracker = ContainerTypeName::FStructTracker();
+template<typename TBaseStruct> FStructTracker TScriptStructTypeBitSet<TBaseStruct>::StructTracker;
\ No newline at end of file
diff --git a/Engine/Plugins/Experimental/StructUtils/Source/StructUtilsTestSuite/Private/ScriptStructTypeBitSetTest.cpp b/Engine/Plugins/Experimental/StructUtils/Source/StructUtilsTestSuite/Private/ScriptStructTypeBitSetTest.cpp
--- a/Engine/Plugins/Experimental/StructUtils/Source/StructUtilsTestSuite/Private/ScriptStructTypeBitSetTest.cpp
+++ b/Engine/Plugins/Experimental/StructUtils/Source/StructUtilsTestSuite/Private/ScriptStructTypeBitSetTest.cpp
@@ -11,10 +11,6 @@
 
 PRAGMA_DISABLE_OPTIMIZATION
 
-template<>
-TScriptStructTypeBitSet<FTestStructSimple>::FStructTracker TScriptStructTypeBitSet<FTestStructSimple>::StructTracker = TScriptStructTypeBitSet<FTestStructSimple>::FStructTracker();
-template struct TScriptStructTypeBitSet<FTestStructSimple>::FStructTracker;
-
 namespace FScriptStructTypeBitSetTests
 {
 
