Added a console command for printing out sizes of all fragment types being used.

diff --git a/Engine/Plugins/Experimental/StructUtils/Source/StructUtils/Public/ScriptStructTypeBitSet.h b/Engine/Plugins/Experimental/StructUtils/Source/StructUtils/Public/ScriptStructTypeBitSet.h
--- a/Engine/Plugins/Experimental/StructUtils/Source/StructUtils/Public/ScriptStructTypeBitSet.h
+++ b/Engine/Plugins/Experimental/StructUtils/Source/StructUtils/Public/ScriptStructTypeBitSet.h
@@ -70,6 +70,8 @@ struct FStructTracker
 		return DebugStructTypeNamesList.IsValidIndex(StructTypeIndex) ? DebugStructTypeNamesList[StructTypeIndex] : FName();
 	}
 
+	TConstArrayView<TWeakObjectPtr<const UScriptStruct>> DebugGetAllStructTypes() const { return StructTypesList; }
+
 	void DebugNukeStructTypeMappingInfo()
 	{
 		StructTypeToIndexSet.Reset();
@@ -422,6 +424,11 @@ public:
 		}
 	}
 
+	static TConstArrayView<TWeakObjectPtr<const UScriptStruct>> DebugGetAllStructTypes()
+	{
+		return StructTracker.DebugGetAllStructTypes();
+	}
+
 	/**
 	 * Resets all the information gathered on the tags. Calling this results in invalidating all previously created
 	 * FStructTypeCollection instances. Used only for debugging and unit/functional testing.
