Fixed alignment problem in the FInstancedStreamStruct

#rnx

diff --git a/Engine/Plugins/Experimental/StructUtils/Source/StructUtils/Public/InstancedStructStream.h b/Engine/Plugins/Experimental/StructUtils/Source/StructUtils/Public/InstancedStructStream.h
--- a/Engine/Plugins/Experimental/StructUtils/Source/StructUtils/Public/InstancedStructStream.h
+++ b/Engine/Plugins/Experimental/StructUtils/Source/StructUtils/Public/InstancedStructStream.h
@@ -359,20 +359,21 @@ protected:
 
 		// Make sure there's space for the struct type and allocation
 		int32 ScriptStructIndex = Head->GetScriptStructIndex(InScriptStruct);
+		uint8* ChunkMemory = reinterpret_cast<uint8*>(Head) + sizeof(FChunkHeader);
 		int32 HeaderOffset = Head->UsedSize;
-		int32 ItemOffset = Align(Head->UsedSize + sizeof(FItemHeader), MinAlignment);
+		int32 ItemOffset = Align(ChunkMemory + Head->UsedSize + sizeof(FItemHeader), MinAlignment) - ChunkMemory;
 		
 		if (ScriptStructIndex == INDEX_NONE || (ItemOffset + ScriptStructSize) > ChunkSize)
 		{
 			AllocateNewChunk();
 			ScriptStructIndex = Head->GetScriptStructIndex(InScriptStruct);
+			ChunkMemory = reinterpret_cast<uint8*>(Head) + sizeof(FChunkHeader);
 			HeaderOffset = Head->UsedSize;
-			ItemOffset = Align(Head->UsedSize + sizeof(FItemHeader), MinAlignment);
+			ItemOffset = Align(ChunkMemory + Head->UsedSize + sizeof(FItemHeader), MinAlignment) - ChunkMemory;
 			check(ScriptStructIndex != INDEX_NONE);
-			check((ItemOffset + ScriptStructSize) <= ChunkSize)
+			check((ItemOffset + ScriptStructSize) <= ChunkSize);
 		}
 
-		uint8* ChunkMemory = reinterpret_cast<uint8*>(Head) + sizeof(FChunkHeader);
 		Head->UsedSize = ItemOffset + ScriptStructSize;
 		Head->NumItems++;
 
@@ -383,6 +384,7 @@ protected:
 
 		NumItems++;
 		
+		check(IsAligned(ChunkMemory + ItemOffset, MinAlignment));
 		return ChunkMemory + ItemOffset;
 	}
 
