I've added a FConstStructView and upgraded a use case we have to accomodate it. I'm sure there are several other places this would be applicable also.

I also removed the unsafe const casts that could be removing const from USTRUCT variables declared as const and causing undefined behavior.

diff --git a/Engine/Plugins/Experimental/StructUtils/Source/StructUtils/Public/InstancedStruct.h b/Engine/Plugins/Experimental/StructUtils/Source/StructUtils/Public/InstancedStruct.h
--- a/Engine/Plugins/Experimental/StructUtils/Source/StructUtils/Public/InstancedStruct.h
+++ b/Engine/Plugins/Experimental/StructUtils/Source/StructUtils/Public/InstancedStruct.h
@@ -9,39 +9,50 @@
 
 #include "InstancedStruct.generated.h"
 
+
+
+///////////////////////////////////////////////////////////////// FConstBaseStruct /////////////////////////////////////////////////////////////////
+
 /**
- * Base functionality for struct pointer handling, 
- * Do not use directly, use either FInstancedStruct, FStructView
+ * Immutable base functionality for struct pointer handling,
+ * Do not use directly, use either FInstancedStruct, FStructView or FConstStructView.
  */
 USTRUCT()
-struct STRUCTUTILS_API FBaseStruct
+struct STRUCTUTILS_API FConstBaseStruct
 {
 	GENERATED_BODY()
 
 public:
 
-	FBaseStruct()
-	: StructMemory(nullptr)
+	FConstBaseStruct() = default;
+
+	/** Copy and move should all be handled by subclasses, but it is impossible to delete them as it is a UStruct... */
+	FConstBaseStruct(const FConstBaseStruct&) { checkf(false, TEXT("Should all be handled by subclasses")); }
+	FConstBaseStruct& operator=(const FConstBaseStruct&) { checkf(false, TEXT("Should all be handled by subclasses")); return *this; }
+	FConstBaseStruct(const FConstBaseStruct&&) { checkf(false, TEXT("Should all be handled by subclasses")); }
+	FConstBaseStruct& operator=(const FConstBaseStruct&&) { checkf(false, TEXT("Should all be handled by subclasses")); return *this; }
+
+	bool operator==(const FConstBaseStruct& Other) const
 	{
+		if ((GetScriptStruct() != Other.GetScriptStruct()) || (GetMemory() != Other.GetMemory()))
+		{
+			return false;
+		}
+		return true;
 	}
 
-	// Copy and move should all be handled by subclasses, but it is impossible to delete them as it is a UStruct...
-	FBaseStruct(const FBaseStruct&) { checkf(false, TEXT("Should all be handled by subclasses")); }
-	FBaseStruct& operator=(const FBaseStruct&) { checkf(false, TEXT("Should all be handled by subclasses")); return *this; }
-	FBaseStruct(const FBaseStruct&&) { checkf(false, TEXT("Should all be handled by subclasses")); }
-	FBaseStruct& operator=(const FBaseStruct&&) { checkf(false, TEXT("Should all be handled by subclasses")); return *this; }
+	bool operator!=(const FConstBaseStruct& Other) const
+	{
+		return !operator==(Other);
+	}
 
-	// Returns mutable reference to the struct, this getter assumes that all data is valid
-	template<typename T>
-	T& GetMutable()
+	/** Returns const pointer to struct memory. */
+	const uint8* GetMemory() const
 	{
-		check(StructMemory != nullptr);
-		check(ScriptStruct != nullptr);
-		check(ScriptStruct->IsChildOf(T::StaticStruct()));
-		return *((T*)StructMemory);
+		return StructMemory;
 	}
 
-	// Returns const reference to the struct, this getter assumes that all data is valid
+	/** Returns const reference to the struct, this getter assumes that all data is valid. */
 	template<typename T>
 	const T& Get() const
 	{
@@ -51,18 +62,7 @@ public:
 		return *((T*)StructMemory);
 	}
 
-	// Returns mutable pointer to the struct, or nullptr if cast is not valid.
-	template<typename T>
-	T* GetMutablePtr()
-	{
-		if (StructMemory != nullptr && ScriptStruct && ScriptStruct->IsChildOf(T::StaticStruct()))
-		{
-			return ((T*)StructMemory);
-		}
-		return nullptr;
-	}
-
-	// Returns const pointer to the struct, or nullptr if cast is not valid.
+	/** Returns const pointer to the struct, or nullptr if cast is not valid. */
 	template<typename T>
 	const T* GetPtr() const
 	{
@@ -73,24 +73,12 @@ public:
 		return nullptr;
 	}
 
-	// Returns struct type.
+	/** Returns struct type. */
 	const UScriptStruct* GetScriptStruct() const
 	{
 		return ScriptStruct;
 	}
 
-	// Returns mutable pointer to struct memory.
-	uint8* GetMutableMemory()
-	{
-		return StructMemory;
-	}
-
-	// Returns const pointer to struct memory.
-	const uint8* GetMemory() const
-	{
-		return StructMemory;
-	}
-
 	/**
 	 * @return True if the struct is valid.
 	 */
@@ -112,13 +100,13 @@ public:
 
 protected:
 
-	FBaseStruct(const UScriptStruct* InScriptStruct, uint8* InStructMemory = nullptr)
+	FConstBaseStruct(const UScriptStruct* InScriptStruct, const uint8* InStructMemory = nullptr)
 		: ScriptStruct(InScriptStruct)
 		, StructMemory(InStructMemory)
 	{
 	}
 
-	void SetMemory(uint8* InStructMemory)
+	void SetMemory(const uint8* InStructMemory)
 	{
 		StructMemory = InStructMemory;
 	}
@@ -128,20 +116,84 @@ protected:
 		ScriptStruct = InScriptStruct;
 	}
 
-    void DestroyScriptStruct()
+	void SetStructData(const UScriptStruct* InScriptStruct, const uint8* InStructMemory)
+	{
+		ScriptStruct = InScriptStruct;
+		StructMemory = InStructMemory;
+	}
+
+protected:
+	const UScriptStruct* ScriptStruct = nullptr;
+	const uint8* StructMemory = nullptr;
+};
+
+///////////////////////////////////////////////////////////////// FBaseStruct /////////////////////////////////////////////////////////////////
+
+/**
+ * Base functionality for mutable struct pointer handling. 
+ * The idea here is that it's only possible for FBaseStruct to be setup from mutable non const memory. This makes the const_cast here safe in GetMutableMemory().
+ * Do not use directly, use either FInstancedStruct, or FStructView.
+ * Note const FBaseStruct only makes the members of this class immutable NOT the struct data pointed at immutable.
+ */
+USTRUCT()
+struct FBaseStruct : public FConstBaseStruct
+{
+	GENERATED_BODY()
+
+public:
+
+	FBaseStruct() = default;
+
+	/** Copy and move should all be handled by subclasses, but it is impossible to delete them as it is a UStruct... */
+	FBaseStruct(const FBaseStruct&) { checkf(false, TEXT("Should all be handled by subclasses")); }
+	FBaseStruct& operator=(const FBaseStruct&) { checkf(false, TEXT("Should all be handled by subclasses")); return *this; }
+	FBaseStruct(const FBaseStruct&&) { checkf(false, TEXT("Should all be handled by subclasses")); }
+	FBaseStruct& operator=(const FBaseStruct&&) { checkf(false, TEXT("Should all be handled by subclasses")); return *this; }
+
+	/** Returns a mutable pointer to struct memory. This const_cast here is safe as a FBaseStruct can only be setup from mutable non const memory. */
+	uint8* GetMutableMemory() const
+	{
+		return const_cast<uint8*>(StructMemory);
+	}
+
+	/** Returns mutable reference to the struct, this getter assumes that all data is valid. */
+	template<typename T>
+	T& GetMutable() const
 	{
 		check(StructMemory != nullptr);
-		if (ScriptStruct != nullptr)
+		check(ScriptStruct != nullptr);
+		check(ScriptStruct->IsChildOf(T::StaticStruct()));
+		return *((T*)GetMutableMemory());
+	}
+
+	/** Returns mutable pointer to the struct, or nullptr if cast is not valid. */
+	template<typename T>
+	T* GetMutablePtr() const
+	{
+		if (StructMemory != nullptr && ScriptStruct && ScriptStruct->IsChildOf(T::StaticStruct()))
 		{
-			ScriptStruct->DestroyStruct(StructMemory);
+			return ((T*)GetMutableMemory());
 		}
+		return nullptr;
 	}
 
-private:
-	const UScriptStruct* ScriptStruct = nullptr;	// Struct type
-	uint8* StructMemory = nullptr;					// Struct memory
+protected:
+	FBaseStruct(const UScriptStruct* InScriptStruct, uint8* InStructMemory = nullptr)
+		: FConstBaseStruct(InScriptStruct, InStructMemory)
+	{}
+
+	void DestroyScriptStruct() const
+	{
+		check(StructMemory != nullptr);
+		if (ScriptStruct != nullptr)
+		{
+			ScriptStruct->DestroyStruct(GetMutableMemory());
+		}
+	}
 };
 
+///////////////////////////////////////////////////////////////// FInstancedStruct /////////////////////////////////////////////////////////////////
+
 /**
  * FInstancedStruct works similarly as instanced UObject* property but is USTRUCTs.
  * Example:
@@ -151,12 +203,7 @@ private:
  *
  *	UPROPERTY(EditAnywhere, Category = Foo, meta = (BaseStruct = "TestStructBase"))
  *	TArray<FInstancedStruct> TestArray;
- *
- * -----------------------------------------------------------------------------
- * DISCLAIMER: This will be removed when TUniquePtr<FScriptObject> is available!
- * -----------------------------------------------------------------------------
  */
-
 USTRUCT()
 struct STRUCTUTILS_API FInstancedStruct : public FBaseStruct
 {
@@ -175,7 +222,7 @@ public:
 		InitializeAs(InScriptStruct, nullptr);
 	}
 
-	FInstancedStruct(const FBaseStruct& InOther)
+	FInstancedStruct(const FConstBaseStruct& InOther)
 		: FBaseStruct()
 	{
 		InitializeAs(InOther.GetScriptStruct(), InOther.GetMemory());
@@ -200,9 +247,9 @@ public:
 		Reset();
 	}
 
-	FInstancedStruct& operator=(const FBaseStruct& InOther)
+	FInstancedStruct& operator=(const FConstBaseStruct& InOther)
 	{
-		if (this != &InOther)
+		if (*this != InOther)
 		{
 			InitializeAs(InOther.GetScriptStruct(), InOther.GetMemory());
 		}
@@ -233,10 +280,10 @@ public:
 		return *this;
 	}
 
-	// Initializes from struct type and optional data.
+	/** Initializes from struct type and optional data. */
 	void InitializeAs(const UScriptStruct* InScriptStruct, const uint8* InStructMemory = nullptr);
 
-	// Initializes from struct type and emplace construct
+	/** Initializes from struct type and emplace construct. */
 	template<typename T, typename... TArgs>
 	void InitializeAs(TArgs&&... InArgs)
 	{
@@ -251,7 +298,7 @@ public:
 		new (GetMutableMemory()) T(Forward<TArgs>(InArgs)...);
 	}
 
-	// Creates a new FInstancedStruct from templated struct type.
+	/** Creates a new FInstancedStruct from templated struct type. */
 	template<typename T>
 	static FInstancedStruct Make()
 	{
@@ -260,7 +307,7 @@ public:
 		return InstancedStruct;
 	}
 
-	// Creates a new FInstancedStruct from templated struct.
+	/** Creates a new FInstancedStruct from templated struct. */
 	template<typename T>
 	static FInstancedStruct Make(const T& Struct)
 	{
@@ -269,7 +316,7 @@ public:
 		return InstancedStruct;
 	}
 
-	// Creates a new FInstancedStruct from the templated type and forward all arguments to constructor
+	/** Creates a new FInstancedStruct from the templated type and forward all arguments to constructor. */
 	template<typename T, typename... TArgs>
 	static inline FInstancedStruct Make(TArgs&&... InArgs)
 	{
@@ -278,10 +325,10 @@ public:
 		return InstancedStruct;
 	}
 
-	// Reset to empty.
+	/** Reset to empty. */
 	void Reset();
 
-	// For StructOpsTypeTraits
+	/** For StructOpsTypeTraits */
 	bool Serialize(FArchive& Ar);
 	bool Identical(const FInstancedStruct* Other, uint32 PortFlags) const;
 	void AddStructReferencedObjects(class FReferenceCollector& Collector);
@@ -289,7 +336,7 @@ public:
 	bool ImportTextItem(const TCHAR*& Buffer, int32 PortFlags, UObject* Parent, FOutputDevice* ErrorText, FArchive* InSerializingArchive = nullptr);
 
 private:
-	// Initializes for new struct type (does nothing if same type) and returns mutable struct
+	/** Initializes for new struct type (does nothing if same type) and returns mutable struct. */
 	UScriptStruct* ReinitializeAs(const UScriptStruct* InScriptStruct);
 };
 
@@ -306,80 +353,103 @@ struct TStructOpsTypeTraits<FInstancedStruct> : public TStructOpsTypeTraitsBase2
 	};
 };
 
-/** 
- * FStructView is "typed" struct pointer, it contains pointer to struct plus UScriptStruct pointer.
- * FStructView does not own the memory and will not free it when out of scope.
+///////////////////////////////////////////////////////////////// FConstStructView /////////////////////////////////////////////////////////////////
+
+/**
+ * FConstStructView is "typed" struct pointer, it contains const pointer to struct plus UScriptStruct pointer.
+ * FConstStructView does not own the memory and will not free it when out of scope.
  * It should be only used to pass struct pointer in a limited scope, or when the user controls the lifetime of the struct being stored.
- * I.e. instead of passing ref or pointer to FInstancedStruct, you should use FStructView to pass around view to the contents.
- * FStructView is passed by value.
- * FStructView is similar to FStructOnScope, but FStructView is a view only (FStructOnScope can either own the memory or be a view)
- *
- * -----------------------------------------------------------------------------
- * DISCLAIMER: This will be removed when FScriptObject is available!
- * -----------------------------------------------------------------------------
+ * E.g. instead of passing ref or pointer to a FInstancedStruct, you should use FConstStructView or FStructView to pass around a view to the contents.
+ * FConstStructView is passed by value.
+ * FConstStructView is similar to FStructOnScope, but FConstStructView is a view only (FStructOnScope can either own the memory or be a view)
  */
-struct STRUCTUTILS_API FStructView : public FBaseStruct
+struct STRUCTUTILS_API FConstStructView : public FConstBaseStruct
 {
 
 public:
 
-	FStructView()
-		: FBaseStruct()
-	{
-	}
+	FConstStructView() = default;
 
-	FStructView(const UScriptStruct* InScriptStruct, const uint8* InStructMemory)
-		: FBaseStruct(InScriptStruct, const_cast<uint8*>(InStructMemory))
+	FConstStructView(const UScriptStruct* InScriptStruct, const uint8* InStructMemory)
+		: FConstBaseStruct(InScriptStruct, InStructMemory)
 	{}
 
-	FStructView(const FBaseStruct& InOther)
-		: FStructView(InOther.GetScriptStruct(), InOther.GetMemory())
+	FConstStructView(const FConstBaseStruct& ConstBaseStruct)
+		: FConstBaseStruct(ConstBaseStruct.GetScriptStruct(), ConstBaseStruct.GetMemory())
+	{}
+
+	FConstStructView(const FConstStructView& Other)
+		: FConstBaseStruct(Other.GetScriptStruct(), Other.GetMemory())
+	{}
+
+	FConstStructView& operator=(const FConstBaseStruct& ConstBaseStruct)
 	{
+		SetStructData(ConstBaseStruct.GetScriptStruct(), ConstBaseStruct.GetMemory());
+		return *this;
 	}
 
-	FStructView(const FStructView& InOther)
-		: FStructView(InOther.GetScriptStruct(), InOther.GetMemory())
+	FConstStructView& operator=(const FConstStructView& Other)
 	{
+		SetStructData(Other.GetScriptStruct(), Other.GetMemory());
+		return *this;
 	}
 
-	FStructView(FStructView&& InOther)
-		: FStructView(InOther.GetScriptStruct(), InOther.GetMemory())
+	/** Creates a new FStructView from the templated struct */
+	template<typename T>
+	static FConstStructView Make(const T& Struct)
 	{
+		return FConstStructView(T::StaticStruct(), reinterpret_cast<const uint8*>(&Struct));
 	}
+};
+
+///////////////////////////////////////////////////////////////// FStructView /////////////////////////////////////////////////////////////////
 
-	~FStructView()
+/**
+ * FStructView is "typed" struct pointer, it contains pointer to struct plus UScriptStruct pointer.
+ * FStructView does not own the memory and will not free it when out of scope.
+ * It should be only used to pass struct pointer in a limited scope, or when the user controls the lifetime of the struct being stored.
+ * E.g. instead of passing ref or pointer to a FInstancedStruct, you should use FConstStructView or FStructView to pass around a view to the contents.
+ * FStructView is passed by value.
+ * FStructView is similar to FStructOnScope, but FStructView is a view only (FStructOnScope can either own the memory or be a view)
+ */
+struct STRUCTUTILS_API FStructView : public FBaseStruct
+{
+
+public:
+
+	FStructView()
+		: FBaseStruct()
 	{
-		Reset();
 	}
 
-	FStructView& operator=(const FBaseStruct& InOther)
-	{
-		SetScriptStruct(InOther.GetScriptStruct());
-		SetMemory(const_cast<uint8*>(InOther.GetMemory()));
+	FStructView(const UScriptStruct* InScriptStruct, uint8* InStructMemory)
+		: FBaseStruct(InScriptStruct, InStructMemory)
+	{}
 
-		return *this;
-	}
+	FStructView(const FBaseStruct& BaseStruct)
+		: FBaseStruct(BaseStruct.GetScriptStruct(), BaseStruct.GetMutableMemory())
+	{}
 
-	FStructView& operator=(const FStructView& InOther)
-	{
-		SetScriptStruct(InOther.GetScriptStruct());
-		SetMemory(const_cast<uint8*>(InOther.GetMemory()));
+	FStructView(const FStructView& Other)
+		: FBaseStruct(Other.GetScriptStruct(), Other.GetMutableMemory())
+	{}
 
+	FStructView& operator=(const FBaseStruct& BaseStruct)
+	{
+		SetStructData(BaseStruct.GetScriptStruct(), BaseStruct.GetMutableMemory());
 		return *this;
 	}
 
-	FStructView& operator=(FStructView&& InOther)
+	FStructView& operator=(const FStructView& Other)
 	{
-		SetScriptStruct(InOther.GetScriptStruct());
-		SetMemory(const_cast<uint8*>(InOther.GetMemory()));
-
+		SetStructData(Other.GetScriptStruct(), Other.GetMutableMemory());
 		return *this;
 	}
 
-	// Creates a new FStructView from the templated struct
+	/** Creates a new FStructView from the templated struct. Note its not safe to make InStruct const ref as the original object may have been declared const */
 	template<typename T>
-	static FStructView Make(const T& Struct)
+	static FStructView Make(T& InStruct)
 	{
-		return FStructView(T::StaticStruct(), reinterpret_cast<const uint8*>(&Struct));
+		return FStructView(T::StaticStruct(), reinterpret_cast<uint8*>(&InStruct));
 	}
 };
\ No newline at end of file
