Fixed InstancedStructStream not handling appeding if the stream itself is empty.

diff --git a/Engine/Plugins/Experimental/StructUtils/Source/StructUtils/Public/InstancedStructStream.h b/Engine/Plugins/Experimental/StructUtils/Source/StructUtils/Public/InstancedStructStream.h
--- a/Engine/Plugins/Experimental/StructUtils/Source/StructUtils/Public/InstancedStructStream.h
+++ b/Engine/Plugins/Experimental/StructUtils/Source/StructUtils/Public/InstancedStructStream.h
@@ -69,10 +69,20 @@ struct STRUCTUTILS_API FInstancedStructStream
 	{
 		if (Other.ChunkSize == ChunkSize)
 		{
-			// If same chunk size, steal chunks.
-			NumItems += Other.NumItems;
-			Head->Next = Other.First;
-			Head = Other.Head;
+			if (NumItems == 0)
+			{
+				// if empty, just hijack the other container
+				NumItems = Other.NumItems;
+				First = Other.First;
+				Head = Other.Head;
+			}
+			else
+			{
+				// If same chunk size, steal chunks.
+				NumItems += Other.NumItems;
+				Head->Next = Other.First;
+				Head = Other.Head;
+			}
 
 			Other.NumItems = 0;
 			Other.First = nullptr;
