Fixed a bug in InstancedStructStream.Append resulting in messing up its internals when appening empty streams.

Mikko.Mononen

diff --git a/Engine/Plugins/Experimental/StructUtils/Source/StructUtils/Public/InstancedStructStream.h b/Engine/Plugins/Experimental/StructUtils/Source/StructUtils/Public/InstancedStructStream.h
--- a/Engine/Plugins/Experimental/StructUtils/Source/StructUtils/Public/InstancedStructStream.h
+++ b/Engine/Plugins/Experimental/StructUtils/Source/StructUtils/Public/InstancedStructStream.h
@@ -67,6 +67,11 @@ struct STRUCTUTILS_API FInstancedStructStream
 	
 	void Append(FInstancedStructStream&& Other)
 	{
+		if (Other.NumItems == 0)
+		{
+			return;
+		}
+
 		if (Other.ChunkSize == ChunkSize)
 		{
 			if (NumItems == 0)
