Removed FArchetypeData.ExportComponentTypes and all it's unused callers.

diff --git a/Engine/Plugins/Experimental/StructUtils/Source/StructUtils/Public/ScriptStructTypeBitSet.h b/Engine/Plugins/Experimental/StructUtils/Source/StructUtils/Public/ScriptStructTypeBitSet.h
--- a/Engine/Plugins/Experimental/StructUtils/Source/StructUtils/Public/ScriptStructTypeBitSet.h
+++ b/Engine/Plugins/Experimental/StructUtils/Source/StructUtils/Public/ScriptStructTypeBitSet.h
@@ -424,6 +424,20 @@ public:
 		}
 	}
 
+	void DebugGetStructTypes(TArray<const UScriptStruct*>& OutComponentList) const
+	{
+		for (int32 Index = 0; Index < StructTypesBitArray.Num(); ++Index)
+		{
+			if (StructTypesBitArray[Index])
+			{
+				if (const UScriptStruct* StructType = StructTracker.GetStructType(Index))
+				{
+					OutComponentList.Add(StructType);
+				}
+			}
+		}
+	}
+
 	static TConstArrayView<TWeakObjectPtr<const UScriptStruct>> DebugGetAllStructTypes()
 	{
 		return StructTracker.DebugGetAllStructTypes();
