Fixed ScriptStructTypeBitSet's hash calculation so that trailing 0 not affect the hash value.

Also, added a test that verifies that behavior.

diff --git a/Engine/Plugins/Experimental/StructUtils/Source/StructUtils/Public/ScriptStructTypeBitSet.h b/Engine/Plugins/Experimental/StructUtils/Source/StructUtils/Public/ScriptStructTypeBitSet.h
--- a/Engine/Plugins/Experimental/StructUtils/Source/StructUtils/Public/ScriptStructTypeBitSet.h
+++ b/Engine/Plugins/Experimental/StructUtils/Source/StructUtils/Public/ScriptStructTypeBitSet.h
@@ -173,9 +173,19 @@ private:
 		{
 			FConstWordIterator Iterator(Instance);
 			uint32 Hash = 0;
+			uint32 TrailingZeroHash = 0;
 			while (Iterator)
 			{
-				Hash = HashCombine(Hash, Iterator.GetWord());
+				const uint32 Word = Iterator.GetWord();
+				if (Word)
+				{
+					Hash = HashCombine(TrailingZeroHash ? TrailingZeroHash : Hash, Word);
+					TrailingZeroHash = 0;
+				}
+				else // potentially a trailing 0-word
+				{
+					TrailingZeroHash = HashCombine(TrailingZeroHash ? TrailingZeroHash : Hash, Word);
+				}
 				++Iterator;
 			}
 			return Hash;
@@ -469,4 +479,4 @@ private:
 	FBitArrayExt StructTypesBitArray;
 };
 
-template<typename TBaseStruct> FStructTracker TScriptStructTypeBitSet<TBaseStruct>::StructTracker;
\ No newline at end of file
+template<typename TBaseStruct> FStructTracker TScriptStructTypeBitSet<TBaseStruct>::StructTracker;
diff --git a/Engine/Plugins/Experimental/StructUtils/Source/StructUtilsTestSuite/Private/ScriptStructTypeBitSetTest.cpp b/Engine/Plugins/Experimental/StructUtils/Source/StructUtilsTestSuite/Private/ScriptStructTypeBitSetTest.cpp
--- a/Engine/Plugins/Experimental/StructUtils/Source/StructUtilsTestSuite/Private/ScriptStructTypeBitSetTest.cpp
+++ b/Engine/Plugins/Experimental/StructUtils/Source/StructUtilsTestSuite/Private/ScriptStructTypeBitSetTest.cpp
@@ -215,18 +215,28 @@ struct FStructUtilsTest_BitSetHash : FAITestBase
 {
 	virtual bool InstantTest() override
 	{
+		FTestStructBitSet EmptyCollection;
 		FTestStructBitSet CollectionA;
 		FTestStructBitSet CollectionB;
 
 		CollectionA.AddBit(9);
 		CollectionB.AddBit(9);
-		CollectionB.AddBit(21);
+		CollectionB.AddBit(1024);
 
 		const uint32 HashA = GetTypeHash(CollectionA);
 		const uint32 HashB = GetTypeHash(CollectionB);
 
 		AITEST_NOT_EQUAL("Two distinct bit sets should have distinct hashes", HashA, HashB);
 
+		CollectionB.RemoveBit(1024);
+		const uint32 HashB2 = GetTypeHash(CollectionB);
+		AITEST_EQUAL("Two bit sets of the same composition should have have identical hashes", HashA, HashB2);
+				
+		CollectionB.RemoveBit(9);
+		const uint32 HashEmpty = GetTypeHash(EmptyCollection);
+		const uint32 HashEmptyB = GetTypeHash(CollectionB);
+		AITEST_EQUAL("An emptied bit set needs to have the same hash as an empty bit set", HashEmpty, HashEmptyB);
+
 		return true;
 	}
 };
