Fixed "truncation warnings" in InstancedStructStream.h

diff --git a/Engine/Plugins/Experimental/StructUtils/Source/StructUtils/Public/InstancedStructStream.h b/Engine/Plugins/Experimental/StructUtils/Source/StructUtils/Public/InstancedStructStream.h
--- a/Engine/Plugins/Experimental/StructUtils/Source/StructUtils/Public/InstancedStructStream.h
+++ b/Engine/Plugins/Experimental/StructUtils/Source/StructUtils/Public/InstancedStructStream.h
@@ -376,7 +376,7 @@ protected:
 		int32 ScriptStructIndex = Head->GetScriptStructIndex(InScriptStruct);
 		uint8* ChunkMemory = reinterpret_cast<uint8*>(Head) + sizeof(FChunkHeader);
 		int32 HeaderOffset = Head->UsedSize;
-		int32 ItemOffset = Align(ChunkMemory + Head->UsedSize + sizeof(FItemHeader), MinAlignment) - ChunkMemory;
+		int32 ItemOffset = int32(Align(ChunkMemory + Head->UsedSize + sizeof(FItemHeader), MinAlignment) - ChunkMemory);
 		
 		if (ScriptStructIndex == INDEX_NONE || (ItemOffset + ScriptStructSize) > ChunkSize)
 		{
@@ -384,7 +384,7 @@ protected:
 			ScriptStructIndex = Head->GetScriptStructIndex(InScriptStruct);
 			ChunkMemory = reinterpret_cast<uint8*>(Head) + sizeof(FChunkHeader);
 			HeaderOffset = Head->UsedSize;
-			ItemOffset = Align(ChunkMemory + Head->UsedSize + sizeof(FItemHeader), MinAlignment) - ChunkMemory;
+			ItemOffset = int32(Align(ChunkMemory + Head->UsedSize + sizeof(FItemHeader), MinAlignment) - ChunkMemory);
 			check(ScriptStructIndex != INDEX_NONE);
 			check((ItemOffset + ScriptStructSize) <= ChunkSize);
 		}
@@ -394,8 +394,8 @@ protected:
 
 		FItemHeader& ItemHeader = *reinterpret_cast<FItemHeader*>(ChunkMemory + HeaderOffset);
 		ItemHeader.Size = (ItemOffset + ScriptStructSize) - HeaderOffset;
-		ItemHeader.Offset = ItemOffset - HeaderOffset;
-		ItemHeader.ScriptStructIndex = ScriptStructIndex;
+		ItemHeader.Offset = int16(ItemOffset - HeaderOffset);
+		ItemHeader.ScriptStructIndex = int16(ScriptStructIndex);
 
 		NumItems++;
 		
