Move the predicate to find s specific type of struct inside a FInstancedStruct Container.

#rnx

diff --git a/Engine/Plugins/Experimental/StructUtils/Source/StructUtils/Public/InstancedStruct.h b/Engine/Plugins/Experimental/StructUtils/Source/StructUtils/Public/InstancedStruct.h
--- a/Engine/Plugins/Experimental/StructUtils/Source/StructUtils/Public/InstancedStruct.h
+++ b/Engine/Plugins/Experimental/StructUtils/Source/StructUtils/Public/InstancedStruct.h
@@ -452,4 +452,15 @@ public:
 	{
 		return FStructView(T::StaticStruct(), reinterpret_cast<uint8*>(&InStruct));
 	}
+};
+
+///////////////////////////////////////////////////////////////// FSameTypeScriptStructPredicate /////////////////////////////////////////////////////////////////
+
+/* Predicate useful to find a struct of a specific type in an container */
+struct FSameTypeScriptStructPredicate
+{
+	const UScriptStruct* TypePtr;
+	FSameTypeScriptStructPredicate(const UScriptStruct* InTypePtr) : TypePtr(InTypePtr) {}
+	FSameTypeScriptStructPredicate(const FConstBaseStruct& InRef) : TypePtr(InRef.GetScriptStruct()) {}
+	bool operator()(const FConstBaseStruct& Other) const { return Other.GetScriptStruct() == TypePtr; }
 };
\ No newline at end of file
diff --git a/Engine/Plugins/Runtime/MassEntity/Source/MassEntity/Private/MassArchetypeData.cpp b/Engine/Plugins/Runtime/MassEntity/Source/MassEntity/Private/MassArchetypeData.cpp
--- a/Engine/Plugins/Runtime/MassEntity/Source/MassEntity/Private/MassArchetypeData.cpp
+++ b/Engine/Plugins/Runtime/MassEntity/Source/MassEntity/Private/MassArchetypeData.cpp
@@ -86,8 +86,7 @@ void FMassArchetypeData::Initialize(const FMassArchetypeCompositionDescriptor& I
 	for (const UScriptStruct* ChunkFragmentType : ChunkFragmentList)
 	{
 		check(ChunkFragmentType);
-		if (const FInstancedStruct* Struct = InInitialValues.ChunkFragments.FindByPredicate( 
-				[ChunkFragmentType](const FInstancedStruct& Struct) { return Struct.GetScriptStruct() == ChunkFragmentType; }))
+		if (const FInstancedStruct* Struct = InInitialValues.ChunkFragments.FindByPredicate(FSameTypeScriptStructPredicate(ChunkFragmentType)))
 		{
 			InitialValues.ChunkFragments.Add(*Struct);
 		}
diff --git a/Engine/Plugins/Runtime/MassGameplay/Source/MassCommon/Public/MassCommonTypes.h b/Engine/Plugins/Runtime/MassGameplay/Source/MassCommon/Public/MassCommonTypes.h
--- a/Engine/Plugins/Runtime/MassGameplay/Source/MassCommon/Public/MassCommonTypes.h
+++ b/Engine/Plugins/Runtime/MassGameplay/Source/MassCommon/Public/MassCommonTypes.h
@@ -51,14 +51,6 @@ struct MASSCOMMON_API FMassUniqueFragmentCollection
 {
 	GENERATED_BODY()
 
-	struct FSameTypeScriptStructPredicate
-	{
-		const UScriptStruct* TypePtr;
-		explicit FSameTypeScriptStructPredicate(const FInstancedStruct& InRef) : TypePtr(InRef.GetScriptStruct()) {}
-		explicit FSameTypeScriptStructPredicate(const UScriptStruct* InTypePtr) : TypePtr(InTypePtr) {}
-		bool operator()(const FInstancedStruct& Other) const { return Other.GetScriptStruct() == TypePtr; }
-	};
-
 	FMassUniqueFragmentCollection() = default;
 
 	/** Adds a unique fragment instance. If a fragment of this type has already been added then its values won't be 
