SharedStruct implementation, similar to TSharedPtr for Objects but now for structs.

#rnx

diff --git a/Engine/Plugins/Experimental/StructUtils/Source/StructUtils/Private/InstancedStruct.cpp b/Engine/Plugins/Experimental/StructUtils/Source/StructUtils/Private/InstancedStruct.cpp
--- a/Engine/Plugins/Experimental/StructUtils/Source/StructUtils/Private/InstancedStruct.cpp
+++ b/Engine/Plugins/Experimental/StructUtils/Source/StructUtils/Private/InstancedStruct.cpp
@@ -1,5 +1,8 @@
 // Copyright Epic Games, Inc. All Rights Reserved.
 #include "InstancedStruct.h"
+#include "StructUtilsTypes.h"
+
+///////////////////////////////////////////////////////////////// FInstancedStruct /////////////////////////////////////////////////////////////////
 
 void FInstancedStruct::InitializeAs(const UScriptStruct* InScriptStruct, const uint8* InStructMemory /*= nullptr*/)
 {
@@ -206,24 +209,18 @@ bool FInstancedStruct::Identical(const FInstancedStruct* Other, uint32 PortFlags
 
 void FInstancedStruct::AddStructReferencedObjects(class FReferenceCollector& Collector)
 {
-	if (const UScriptStruct* ScriptStructPtr = GetScriptStruct())
-	{
-		Collector.AddReferencedObject(ScriptStructPtr);
-
-		if (ScriptStructPtr->StructFlags & STRUCT_AddStructReferencedObjects)
-		{
-			ScriptStructPtr->GetCppStructOps()->AddStructReferencedObjects()(GetMutableMemory(), Collector);
-		}
-		else
-		{
-			// The iterator will recursively loop through all structs in structs too.
-			for (TPropertyValueIterator<const FObjectProperty> It(ScriptStructPtr, GetMemory()); It; ++It)
-			{
-				UObject** ObjectPtr = static_cast<UObject**>(const_cast<void*>(It.Value()));
-				Collector.AddReferencedObject(*ObjectPtr);
-			}
-		}
-	}
+	UE::StructUtils::AddStructReferencedObjects(*this, Collector);
 }
 
+///////////////////////////////////////////////////////////////// FConstSharedStruct /////////////////////////////////////////////////////////////////
+
+bool FConstSharedStruct::Identical(const FConstSharedStruct* Other, uint32 PortFlags) const
+{
+	// Only empty is considered equal
+	return Other != nullptr && GetMemory() == nullptr && Other->GetMemory() == nullptr && GetScriptStruct() == nullptr && Other->GetScriptStruct() == nullptr;
+}
 
+void FConstSharedStruct::AddStructReferencedObjects(class FReferenceCollector& Collector)
+{
+	UE::StructUtils::AddStructReferencedObjects(*this, Collector);
+}
diff --git a/Engine/Plugins/Experimental/StructUtils/Source/StructUtils/Private/StructUtilsTypes.cpp b/Engine/Plugins/Experimental/StructUtils/Source/StructUtils/Private/StructUtilsTypes.cpp
new file mode 100644
--- /dev/null
+++ b/Engine/Plugins/Experimental/StructUtils/Source/StructUtils/Private/StructUtilsTypes.cpp
@@ -0,0 +1,57 @@
+// Copyright Epic Games, Inc. All Rights Reserved.
+
+#pragma once
+
+#include "StructUtilsTypes.h"
+#include "UObject/Class.h"
+#include "Serialization/ArchiveCrc32.h"
+#include "InstancedStruct.h"
+
+namespace UE::StructUtils
+{
+	STRUCTUTILS_API void AddStructReferencedObjects(const FConstStructView& StructView, class FReferenceCollector& Collector)
+	{
+		if (const UScriptStruct* ScriptStructPtr = StructView.GetScriptStruct())
+		{
+			Collector.AddReferencedObject(ScriptStructPtr);
+
+			if (uint8* StructMemory = const_cast<uint8*>(StructView.GetMemory()))
+			{
+				if (ScriptStructPtr->StructFlags & STRUCT_AddStructReferencedObjects)
+				{
+					ScriptStructPtr->GetCppStructOps()->AddStructReferencedObjects()(StructMemory, Collector);
+				}
+				else
+				{
+					// The iterator will recursively loop through all structs in structs too.
+					for (TPropertyValueIterator<const FObjectProperty> It(ScriptStructPtr, StructMemory); It; ++It)
+					{
+						UObject** ObjectPtr = static_cast<UObject**>(const_cast<void*>(It.Value()));
+						Collector.AddReferencedObject(*ObjectPtr);
+					}
+				}
+			}
+		}
+	}
+
+	STRUCTUTILS_API uint32 GetStructCrc32(const UScriptStruct& ScriptStruct, const uint8* StructMemory, const uint32 CRC /*= 0*/)
+	{
+		FArchiveCrc32 Ar(HashCombine(CRC, PointerHash(&ScriptStruct)));
+		if (StructMemory)
+		{
+			UScriptStruct& NonConstScriptStruct = const_cast<UScriptStruct&>(ScriptStruct);
+			NonConstScriptStruct.SerializeItem(Ar, const_cast<uint8*>(StructMemory), nullptr);
+		}
+		return Ar.GetCrc();
+	}
+
+	STRUCTUTILS_API uint32 GetStructCrc32(const FConstStructView& StructView, const uint32 CRC /*= 0*/)
+	{
+		if (const UScriptStruct* ScriptStruct = StructView.GetScriptStruct())
+		{
+			return GetStructCrc32(*ScriptStruct, StructView.GetMemory(), CRC);
+		}
+		return 0;
+	}
+
+}
diff --git a/Engine/Plugins/Experimental/StructUtils/Source/StructUtils/Public/InstancedStruct.h b/Engine/Plugins/Experimental/StructUtils/Source/StructUtils/Public/InstancedStruct.h
--- a/Engine/Plugins/Experimental/StructUtils/Source/StructUtils/Public/InstancedStruct.h
+++ b/Engine/Plugins/Experimental/StructUtils/Source/StructUtils/Public/InstancedStruct.h
@@ -287,6 +287,8 @@ public:
 	template<typename T, typename... TArgs>
 	void InitializeAs(TArgs&&... InArgs)
 	{
+		static_assert(!TIsDerivedFrom<T, FConstBaseStruct>::IsDerived && !TIsDerivedFrom<T, FConstSharedStruct>::IsDerived, "It does not make sense to intialize a instanced struct over an other struct wrapper type");
+
 		Reset();
 
 		const UScriptStruct* InScriptStruct = T::StaticStruct();
@@ -302,6 +304,8 @@ public:
 	template<typename T>
 	static FInstancedStruct Make()
 	{
+		static_assert(!TIsDerivedFrom<T, FConstBaseStruct>::IsDerived && !TIsDerivedFrom<T, FConstSharedStruct>::IsDerived, "It does not make sense to create a instanced struct over an other struct wrapper type");
+
 		FInstancedStruct InstancedStruct;
 		InstancedStruct.InitializeAs(T::StaticStruct(), nullptr);
 		return InstancedStruct;
@@ -311,6 +315,8 @@ public:
 	template<typename T>
 	static FInstancedStruct Make(const T& Struct)
 	{
+		static_assert(!TIsDerivedFrom<T, FConstBaseStruct>::IsDerived && !TIsDerivedFrom<T, FConstSharedStruct>::IsDerived, "It does not make sense to create a instanced struct over an other struct wrapper type");
+
 		FInstancedStruct InstancedStruct;
 		InstancedStruct.InitializeAs(T::StaticStruct(), reinterpret_cast<const uint8*>(&Struct));
 		return InstancedStruct;
@@ -320,6 +326,8 @@ public:
 	template<typename T, typename... TArgs>
 	static inline FInstancedStruct Make(TArgs&&... InArgs)
 	{
+		static_assert(!TIsDerivedFrom<T, FConstBaseStruct>::IsDerived && !TIsDerivedFrom<T, FConstSharedStruct>::IsDerived, "It does not make sense to create a instanced struct over an other struct wrapper type");
+
 		FInstancedStruct InstancedStruct;
 		InstancedStruct.InitializeAs<T>(Forward<TArgs>(InArgs)...);
 		return InstancedStruct;
@@ -353,6 +361,345 @@ struct TStructOpsTypeTraits<FInstancedStruct> : public TStructOpsTypeTraitsBase2
 	};
 };
 
+///////////////////////////////////////////////////////////////// FStructSharedMemory /////////////////////////////////////////////////////////////////
+
+/**
+ * Holds the information and memory about a UStruct and that is the actual part that is shared across all the FConstSharedStruct/FSharedStruct
+ * 
+ * The size of the allocation for this structure should always includes not only the need size for it members but also the size required to hold the
+ * structure describe by SciprtStruct. This is how we can avoid 2 pointer referencing(cache misses). Look at the Create() method to understand more.
+ */
+struct STRUCTUTILS_API FStructSharedMemory : public TSharedFromThis<FStructSharedMemory>
+{
+	~FStructSharedMemory()
+	{
+		ScriptStruct.DestroyStruct(GetMemory());
+	}
+
+	static TSharedPtr<FStructSharedMemory> Create(const UScriptStruct& InScriptStruct, const uint8* InStructMemory = nullptr)
+	{
+		const int32 RequiredSize = sizeof(FStructSharedMemory) + InScriptStruct.GetStructureSize();
+		FStructSharedMemory* StructMemory = new(new uint8[RequiredSize]) FStructSharedMemory(InScriptStruct, InStructMemory);
+		return MakeShareable(StructMemory);
+	}
+
+	/** Returns pointer to struct memory. */
+	uint8* GetMemory() const
+	{
+		return (uint8*)StructMemory;
+	}
+
+	/** Returns struct type. */
+	const UScriptStruct& GetScriptStruct() const
+	{
+		return ScriptStruct;
+	}
+
+private:
+	FStructSharedMemory(const UScriptStruct& InScriptStruct, const uint8* InStructMemory = nullptr)
+		: ScriptStruct(InScriptStruct)
+	{
+		if (InStructMemory)
+		{
+			ScriptStruct.CopyScriptStruct(StructMemory, InStructMemory);
+		}
+	}
+
+	const UScriptStruct& ScriptStruct;
+
+	// The required memory size for the struct represented by the UScriptStruct must be allocated right after this object into big enough preallocated buffer, 
+	// Check Create() method for more information.
+	uint8 StructMemory[0];
+};
+
+///////////////////////////////////////////////////////////////// FConstSharedStruct /////////////////////////////////////////////////////////////////
+
+/**
+ * FConstSharedStruct is the same as the FSharedStruct but restrict the API to return const struct type. 
+ * 
+ * See FSharedStruct for more information.
+ */
+USTRUCT()
+struct STRUCTUTILS_API FConstSharedStruct
+{
+	GENERATED_BODY();
+
+	FConstSharedStruct()
+	{}
+
+	FConstSharedStruct(const FConstSharedStruct& InOther)
+	{
+		StructMemoryPtr = InOther.StructMemoryPtr;
+	}
+
+	FConstSharedStruct(const FConstSharedStruct&& InOther)
+	{
+		StructMemoryPtr = InOther.StructMemoryPtr;
+	}
+
+	FConstSharedStruct& operator=(const FConstSharedStruct& InOther)
+	{
+		if (this != &InOther)
+		{
+			StructMemoryPtr = InOther.StructMemoryPtr;
+		}
+		return *this;
+	}
+	FConstSharedStruct& operator=(FConstSharedStruct&& InOther)
+	{
+		if (this != &InOther)
+		{
+			StructMemoryPtr = InOther.StructMemoryPtr;
+		}
+		return *this;
+	}
+
+	bool operator==(const FConstSharedStruct& Other) const
+	{
+		if ((GetScriptStruct() != Other.GetScriptStruct()) || (GetMemory() != Other.GetMemory()))
+		{
+			return false;
+		}
+		return true;
+	}
+
+	bool operator!=(const FConstSharedStruct& Other) const
+	{
+		return !operator==(Other);
+	}
+
+	/** Returns const pointer to struct memory. */
+	const uint8* GetMemory() const
+	{
+		return StructMemoryPtr ? StructMemoryPtr.Get()->GetMemory() : nullptr;
+	}
+
+	/** Returns const reference to the struct, this getter assumes that all data is valid. */
+	template<typename T>
+	const T& Get() const
+	{
+		const uint8* Memory = GetMemory();
+		const UScriptStruct* ScriptStruct = GetScriptStruct();
+		check(Memory != nullptr);
+		check(ScriptStruct != nullptr);
+		check(ScriptStruct->IsChildOf(T::StaticStruct()));
+		return *((T*)Memory);
+	}
+
+	/** Returns const pointer to the struct, or nullptr if cast is not valid. */
+	template<typename T>
+	const T* GetPtr() const
+	{
+		const uint8* Memory = GetMemory();
+		const UScriptStruct* ScriptStruct = GetScriptStruct();
+		if (Memory != nullptr && ScriptStruct && ScriptStruct->IsChildOf(T::StaticStruct()))
+		{
+			return ((T*)Memory);
+		}
+		return nullptr;
+	}
+
+	/** Returns struct type. */
+	const UScriptStruct* GetScriptStruct() const
+	{
+		return StructMemoryPtr ? &(StructMemoryPtr.Get()->GetScriptStruct()) : nullptr;
+	}
+
+	/**
+	 * @return True if the struct is valid.
+	 */
+	bool IsValid() const
+	{
+		return GetMemory() != nullptr && GetScriptStruct() != nullptr;
+	}
+
+	/**
+	 * Reset to empty.
+	 */
+	void Reset()
+	{
+		StructMemoryPtr.Reset();
+	}
+
+	/** For StructOpsTypeTraits */
+	bool Identical(const FConstSharedStruct* Other, uint32 PortFlags) const;
+	void AddStructReferencedObjects(class FReferenceCollector& Collector);
+
+protected:
+
+	TSharedPtr<const FStructSharedMemory> StructMemoryPtr;
+};
+
+template<>
+struct TStructOpsTypeTraits<FConstSharedStruct> : public TStructOpsTypeTraitsBase2<FConstSharedStruct>
+{
+	enum
+	{
+		WithIdentical = true,
+		WithAddStructReferencedObjects = true,
+	};
+};
+
+///////////////////////////////////////////////////////////////// FSharedStruct /////////////////////////////////////////////////////////////////
+
+/**
+ * FSharedStruct works similarly as a TSharedPtr<FInstancedStruct> but removes the double pointer indirection that would create.
+ * (One pointer for the FInstancedStruct and one pointer for the struct memory it is wrapping).
+ * Also note that because of its implementation, it is not possible for now to go from a struct reference or struct view back to a shared struct. 
+ * 
+ * This struct type is also convertible to a FStructView and is the preferable way of passing it as a parameter just as the FInstancedStruct.
+ * If the calling code would like to keep a shared pointer to the struct, you may pass the FSharedStruct as a parameter but it is recommended to pass it as 
+ * a "const FSharedStruct&" to limit the unnecessary recounting.
+ * 
+ */
+USTRUCT()
+struct STRUCTUTILS_API FSharedStruct : public FConstSharedStruct
+{
+	GENERATED_BODY();
+
+	FSharedStruct()
+	{
+	}
+
+	explicit FSharedStruct(const UScriptStruct* InScriptStruct)
+	{
+		InitializeAs(InScriptStruct, nullptr);
+	}
+
+	FSharedStruct(const FConstBaseStruct& InOther)
+	{
+		InitializeAs(InOther.GetScriptStruct(), InOther.GetMemory());
+	}
+
+	FSharedStruct(const FSharedStruct& InOther)
+		: FConstSharedStruct(InOther)
+	{
+	}
+
+	FSharedStruct(FSharedStruct&& InOther)
+		: FConstSharedStruct(InOther)
+	{
+	}
+
+	~FSharedStruct()
+	{
+		Reset();
+	}
+
+	FSharedStruct& operator=(const FConstBaseStruct& InOther)
+	{
+		InitializeAs(InOther.GetScriptStruct(), InOther.GetMemory());
+		return *this;
+	}
+
+	FSharedStruct& operator=(const FSharedStruct& InOther)
+	{
+		FConstSharedStruct::operator=(InOther);
+		return *this;
+	}
+
+	FSharedStruct& operator=(FSharedStruct&& InOther)
+	{
+		FConstSharedStruct::operator=(InOther);
+		return *this;
+	}
+
+	/** Returns a mutable pointer to struct memory. This const_cast here is safe as a FSharedStruct can only be setup from mutable non const memory. */
+	uint8* GetMutableMemory() const
+	{
+		return const_cast<uint8*>(GetMemory());
+	}
+
+	/** Returns mutable reference to the struct, this getter assumes that all data is valid. */
+	template<typename T>
+	T& GetMutable() const
+	{
+		uint8* Memory = GetMutableMemory();
+		const UScriptStruct* ScriptStruct = GetScriptStruct();
+		check(Memory != nullptr);
+		check(ScriptStruct != nullptr);
+		check(ScriptStruct->IsChildOf(T::StaticStruct()));
+		return *((T*)Memory);
+	}
+
+	/** Returns mutable pointer to the struct, or nullptr if cast is not valid. */
+	template<typename T>
+	T* GetMutablePtr() const
+	{
+		uint8* Memory = GetMutableMemory();
+		const UScriptStruct* ScriptStruct = GetScriptStruct();
+		if (Memory != nullptr && ScriptStruct && ScriptStruct->IsChildOf(T::StaticStruct()))
+		{
+			return ((T*)Memory);
+		}
+		return nullptr;
+	}
+
+	/** Initializes from struct type and optional data. */
+	void InitializeAs(const UScriptStruct* InScriptStruct, const uint8* InStructMemory = nullptr)
+	{
+		Reset();
+		if (InScriptStruct)
+		{
+			StructMemoryPtr = FStructSharedMemory::Create(*InScriptStruct, InStructMemory);
+		}
+	}
+
+	/** Initializes from struct type and emplace construct. */
+	template<typename T, typename... TArgs>
+	void InitializeAs(TArgs&&... InArgs)
+	{
+		static_assert(!TIsDerivedFrom<T, FConstBaseStruct>::IsDerived && !TIsDerivedFrom<T, FConstSharedStruct>::IsDerived, "It does not make sense to initialize a shared struct over an other struct wrapper type");
+
+		Reset();
+		StructMemoryPtr = FStructSharedMemory::Create(*T::StaticStruct());
+		new (GetMutableMemory()) T(Forward<TArgs>(InArgs)...);
+	}
+
+	/** Creates a new FSharedStruct from templated struct type. */
+	template<typename T>
+	static FSharedStruct Make()
+	{
+		static_assert(!TIsDerivedFrom<T, FConstBaseStruct>::IsDerived && !TIsDerivedFrom<T, FConstSharedStruct>::IsDerived, "It does not make sense to create a shared struct over an other struct wrapper type");
+
+		FSharedStruct SharedStruct;
+		SharedStruct.InitializeAs(T::StaticStruct(), nullptr);
+		return SharedStruct;
+	}
+
+	/** Creates a new FSharedStruct from templated struct. */
+	template<typename T>
+	static FSharedStruct Make(const T& Struct)
+	{
+		static_assert(!TIsDerivedFrom<T, FConstBaseStruct>::IsDerived && !TIsDerivedFrom<T, FConstSharedStruct>::IsDerived, "It does not make sense to create a shared struct over an other struct wrapper type");
+
+		FSharedStruct SharedStruct;
+		SharedStruct.InitializeAs(T::StaticStruct(), reinterpret_cast<const uint8*>(&Struct));
+		return SharedStruct;
+	}
+
+	/** Creates a new FSharedStruct from the templated type and forward all arguments to constructor. */
+	template<typename T, typename... TArgs>
+	static inline FSharedStruct Make(TArgs&&... InArgs)
+	{
+		static_assert(!TIsDerivedFrom<T, FConstBaseStruct>::IsDerived && !TIsDerivedFrom<T, FConstSharedStruct>::IsDerived, "It does not make sense to create a shared struct over an other struct wrapper type");
+
+		FSharedStruct SharedStruct;
+		SharedStruct.InitializeAs<T>(Forward<TArgs>(InArgs)...);
+		return SharedStruct;
+	}
+};
+
+template<>
+struct TStructOpsTypeTraits<FSharedStruct> : public TStructOpsTypeTraitsBase2<FSharedStruct>
+{
+	enum
+	{
+		WithIdentical = true,
+		WithAddStructReferencedObjects = true,
+	};
+};
+
 ///////////////////////////////////////////////////////////////// FConstStructView /////////////////////////////////////////////////////////////////
 
 /**
@@ -370,7 +717,7 @@ public:
 
 	FConstStructView() = default;
 
-	FConstStructView(const UScriptStruct* InScriptStruct, const uint8* InStructMemory)
+	FConstStructView(const UScriptStruct* InScriptStruct, const uint8* InStructMemory = nullptr)
 		: FConstBaseStruct(InScriptStruct, InStructMemory)
 	{}
 
@@ -378,6 +725,10 @@ public:
 		: FConstBaseStruct(ConstBaseStruct.GetScriptStruct(), ConstBaseStruct.GetMemory())
 	{}
 
+	FConstStructView(const FConstSharedStruct& SharedStruct)
+		: FConstBaseStruct(SharedStruct.GetScriptStruct(), SharedStruct.GetMemory())
+	{}
+
 	FConstStructView(const FConstStructView& Other)
 		: FConstBaseStruct(Other.GetScriptStruct(), Other.GetMemory())
 	{}
@@ -398,6 +749,7 @@ public:
 	template<typename T>
 	static FConstStructView Make(const T& Struct)
 	{
+		static_assert(!TIsDerivedFrom<T, FConstBaseStruct>::IsDerived && !TIsDerivedFrom<T, FConstSharedStruct>::IsDerived, "It does not make sense to create a const struct view over an other struct wrapper type");
 		return FConstStructView(T::StaticStruct(), reinterpret_cast<const uint8*>(&Struct));
 	}
 };
@@ -422,7 +774,7 @@ public:
 	{
 	}
 
-	FStructView(const UScriptStruct* InScriptStruct, uint8* InStructMemory)
+	FStructView(const UScriptStruct* InScriptStruct, uint8* InStructMemory = nullptr)
 		: FBaseStruct(InScriptStruct, InStructMemory)
 	{}
 
@@ -430,6 +782,10 @@ public:
 		: FBaseStruct(BaseStruct.GetScriptStruct(), BaseStruct.GetMutableMemory())
 	{}
 
+	FStructView(const FSharedStruct& SharedStruct)
+		: FBaseStruct(SharedStruct.GetScriptStruct(), SharedStruct.GetMutableMemory())
+	{}
+
 	FStructView(const FStructView& Other)
 		: FBaseStruct(Other.GetScriptStruct(), Other.GetMutableMemory())
 	{}
@@ -450,6 +806,7 @@ public:
 	template<typename T>
 	static FStructView Make(T& InStruct)
 	{
+		static_assert(!TIsDerivedFrom<T, FConstBaseStruct>::IsDerived && !TIsDerivedFrom<T, FConstSharedStruct>::IsDerived, "It does not make sense to create a struct view over an other struct wrapper type");
 		return FStructView(T::StaticStruct(), reinterpret_cast<uint8*>(&InStruct));
 	}
 };
@@ -461,6 +818,7 @@ struct FSameTypeScriptStructPredicate
 {
 	const UScriptStruct* TypePtr;
 	FSameTypeScriptStructPredicate(const UScriptStruct* InTypePtr) : TypePtr(InTypePtr) {}
-	FSameTypeScriptStructPredicate(const FConstBaseStruct& InRef) : TypePtr(InRef.GetScriptStruct()) {}
-	bool operator()(const FConstBaseStruct& Other) const { return Other.GetScriptStruct() == TypePtr; }
-};
\ No newline at end of file
+	FSameTypeScriptStructPredicate(const FConstStructView& InRef) : TypePtr(InRef.GetScriptStruct()) {}
+
+	bool operator()(const FConstStructView& Other) const { return Other.GetScriptStruct() == TypePtr; }
+};
diff --git a/Engine/Plugins/Experimental/StructUtils/Source/StructUtils/Public/StructUtilsTypes.h b/Engine/Plugins/Experimental/StructUtils/Source/StructUtils/Public/StructUtilsTypes.h
--- a/Engine/Plugins/Experimental/StructUtils/Source/StructUtils/Public/StructUtilsTypes.h
+++ b/Engine/Plugins/Experimental/StructUtils/Source/StructUtils/Public/StructUtilsTypes.h
@@ -2,6 +2,20 @@
 
 #pragma once
 
+#include "UObject/Class.h"
+
 #ifndef WITH_STRUCTUTILS_DEBUG
 #define WITH_STRUCTUTILS_DEBUG (!(UE_BUILD_SHIPPING || UE_BUILD_SHIPPING_WITH_EDITOR || UE_BUILD_TEST) && 1)
-#endif // WITH_STRUCTUTILS_DEBUG
\ No newline at end of file
+#endif // WITH_STRUCTUTILS_DEBUG
+
+struct FConstStructView;
+class FReferenceCollector;
+
+namespace UE::StructUtils
+{
+	extern STRUCTUTILS_API void AddStructReferencedObjects(const FConstStructView& StructView, class FReferenceCollector& Collector);
+
+	extern STRUCTUTILS_API uint32 GetStructCrc32(const UScriptStruct& ScriptStruct, const uint8* StructMemory, const uint32 CRC = 0);
+
+	extern STRUCTUTILS_API uint32 GetStructCrc32(const FConstStructView& StructView, const uint32 CRC = 0);
+}
