Fix code analysis warning from last submit

#rnx

diff --git a/Engine/Plugins/Experimental/StructUtils/Source/StructUtils/Public/InstancedStruct.h b/Engine/Plugins/Experimental/StructUtils/Source/StructUtils/Public/InstancedStruct.h
--- a/Engine/Plugins/Experimental/StructUtils/Source/StructUtils/Public/InstancedStruct.h
+++ b/Engine/Plugins/Experimental/StructUtils/Source/StructUtils/Public/InstancedStruct.h
@@ -379,6 +379,8 @@ struct STRUCTUTILS_API FStructSharedMemory : public TSharedFromThis<FStructShare
 	static TSharedPtr<FStructSharedMemory> Create(const UScriptStruct& InScriptStruct, const uint8* InStructMemory = nullptr)
 	{
 		const int32 RequiredSize = sizeof(FStructSharedMemory) + InScriptStruct.GetStructureSize();
+		// Code analysis is unable to understand correctly what we are doing here, so disabling the warning C6386: Buffer overrun while writing to...
+		CA_SUPPRESS( 6386 )
 		FStructSharedMemory* StructMemory = new(new uint8[RequiredSize]) FStructSharedMemory(InScriptStruct, InStructMemory);
 		return MakeShareable(StructMemory);
 	}
