Removed pragma once in a cpp file

#rnx

diff --git a/Engine/Plugins/Experimental/StructUtils/Source/StructUtils/Private/StructUtilsTypes.cpp b/Engine/Plugins/Experimental/StructUtils/Source/StructUtils/Private/StructUtilsTypes.cpp
--- a/Engine/Plugins/Experimental/StructUtils/Source/StructUtils/Private/StructUtilsTypes.cpp
+++ b/Engine/Plugins/Experimental/StructUtils/Source/StructUtils/Private/StructUtilsTypes.cpp
@@ -1,7 +1,5 @@
 // Copyright Epic Games, Inc. All Rights Reserved.
 
-#pragma once
-
 #include "StructUtilsTypes.h"
 #include "UObject/Class.h"
 #include "Serialization/ArchiveCrc32.h"
