Flatten instanced struct hierarchy

#rnx

diff --git a/Engine/Plugins/Experimental/StructUtils/Source/StructUtils/Private/InstancedStruct.cpp b/Engine/Plugins/Experimental/StructUtils/Source/StructUtils/Private/InstancedStruct.cpp
--- a/Engine/Plugins/Experimental/StructUtils/Source/StructUtils/Private/InstancedStruct.cpp
+++ b/Engine/Plugins/Experimental/StructUtils/Source/StructUtils/Private/InstancedStruct.cpp
@@ -1,8 +1,21 @@
 // Copyright Epic Games, Inc. All Rights Reserved.
 #include "InstancedStruct.h"
-#include "StructUtilsTypes.h"
+#include "StructView.h"
+#include "SharedStruct.h"
 
-///////////////////////////////////////////////////////////////// FInstancedStruct /////////////////////////////////////////////////////////////////
+FInstancedStruct::FInstancedStruct(const FConstStructView InOther)
+{
+	InitializeAs(InOther.GetScriptStruct(), InOther.GetMemory());
+}
+
+FInstancedStruct& FInstancedStruct::operator=(const FConstStructView InOther)
+{
+	if(*this != InOther)
+	{
+		InitializeAs(InOther.GetScriptStruct(), InOther.GetMemory());
+	}
+	return *this;
+}
 
 void FInstancedStruct::InitializeAs(const UScriptStruct* InScriptStruct, const uint8* InStructMemory /*= nullptr*/)
 {
@@ -14,10 +27,9 @@ void FInstancedStruct::InitializeAs(const UScriptStruct* InScriptStruct, const u
 		return;
 	}
 
-	SetScriptStruct(InScriptStruct);
-
 	const int32 RequiredSize = InScriptStruct->GetStructureSize();
-	SetMemory((uint8*)FMemory::Malloc(FMath::Max(1, RequiredSize)));
+	const uint8* Memory = ((uint8*)FMemory::Malloc(FMath::Max(1, RequiredSize)));
+	SetStructData(InScriptStruct,Memory);
 
 	InScriptStruct->InitializeStruct(GetMutableMemory());
 
@@ -34,7 +46,7 @@ void FInstancedStruct::Reset()
 		DestroyScriptStruct();
 		FMemory::Free(Memory);
 	}
-	FBaseStruct::Reset();
+	ResetStructData();
 }
 
 bool FInstancedStruct::Serialize(FArchive& Ar)
@@ -209,18 +221,8 @@ bool FInstancedStruct::Identical(const FInstancedStruct* Other, uint32 PortFlags
 
 void FInstancedStruct::AddStructReferencedObjects(class FReferenceCollector& Collector)
 {
-	UE::StructUtils::AddStructReferencedObjects(*this, Collector);
-}
-
-///////////////////////////////////////////////////////////////// FConstSharedStruct /////////////////////////////////////////////////////////////////
-
-bool FConstSharedStruct::Identical(const FConstSharedStruct* Other, uint32 PortFlags) const
-{
-	// Only empty is considered equal
-	return Other != nullptr && GetMemory() == nullptr && Other->GetMemory() == nullptr && GetScriptStruct() == nullptr && Other->GetScriptStruct() == nullptr;
-}
-
-void FConstSharedStruct::AddStructReferencedObjects(class FReferenceCollector& Collector)
-{
-	UE::StructUtils::AddStructReferencedObjects(*this, Collector);
-}
+	if (const UScriptStruct* Struct = GetScriptStruct())
+	{
+		Collector.AddReferencedObjects(Struct, GetMutableMemory());
+	}
+}
\ No newline at end of file
diff --git a/Engine/Plugins/Experimental/StructUtils/Source/StructUtils/Private/SharedStruct.cpp b/Engine/Plugins/Experimental/StructUtils/Source/StructUtils/Private/SharedStruct.cpp
new file mode 100644
--- /dev/null
+++ b/Engine/Plugins/Experimental/StructUtils/Source/StructUtils/Private/SharedStruct.cpp
@@ -0,0 +1,37 @@
+// Copyright Epic Games, Inc. All Rights Reserved.
+#include "SharedStruct.h"
+#include "InstancedStruct.h"
+#include "StructView.h"
+#include "StructUtilsTypes.h"
+
+///////////////////////////////////////////////////////////////// FConstSharedStruct /////////////////////////////////////////////////////////////////
+
+bool FConstSharedStruct::Identical(const FConstSharedStruct* Other, uint32 PortFlags) const
+{
+	// Only empty is considered equal
+	return Other != nullptr && GetMemory() == nullptr && Other->GetMemory() == nullptr && GetScriptStruct() == nullptr && Other->GetScriptStruct() == nullptr;
+}
+
+void FConstSharedStruct::AddStructReferencedObjects(class FReferenceCollector& Collector)
+{
+	if (const UScriptStruct* Struct = GetScriptStruct())
+	{
+		Collector.AddReferencedObjects(Struct, const_cast<uint8*>(GetMemory()));
+	}
+}
+
+///////////////////////////////////////////////////////////////// FSharedStruct /////////////////////////////////////////////////////////////////
+
+FSharedStruct::FSharedStruct(const FConstStructView InOther)
+{
+	InitializeAs(InOther.GetScriptStruct(), InOther.GetMemory());
+}
+
+FSharedStruct& FSharedStruct::operator=(const FConstStructView InOther)
+{
+	if (*this != InOther)
+	{
+		InitializeAs(InOther.GetScriptStruct(), InOther.GetMemory());
+	}
+	return *this;
+}
\ No newline at end of file
diff --git a/Engine/Plugins/Experimental/StructUtils/Source/StructUtils/Private/StructUtilsTypes.cpp b/Engine/Plugins/Experimental/StructUtils/Source/StructUtils/Private/StructUtilsTypes.cpp
--- a/Engine/Plugins/Experimental/StructUtils/Source/StructUtils/Private/StructUtilsTypes.cpp
+++ b/Engine/Plugins/Experimental/StructUtils/Source/StructUtils/Private/StructUtilsTypes.cpp
@@ -7,31 +7,6 @@
 
 namespace UE::StructUtils
 {
-	STRUCTUTILS_API void AddStructReferencedObjects(const FConstStructView& StructView, class FReferenceCollector& Collector)
-	{
-		if (const UScriptStruct* ScriptStructPtr = StructView.GetScriptStruct())
-		{
-			Collector.AddReferencedObject(ScriptStructPtr);
-
-			if (uint8* StructMemory = const_cast<uint8*>(StructView.GetMemory()))
-			{
-				if (ScriptStructPtr->StructFlags & STRUCT_AddStructReferencedObjects)
-				{
-					ScriptStructPtr->GetCppStructOps()->AddStructReferencedObjects()(StructMemory, Collector);
-				}
-				else
-				{
-					// The iterator will recursively loop through all structs in structs too.
-					for (TPropertyValueIterator<const FObjectProperty> It(ScriptStructPtr, StructMemory); It; ++It)
-					{
-						UObject** ObjectPtr = static_cast<UObject**>(const_cast<void*>(It.Value()));
-						Collector.AddReferencedObject(*ObjectPtr);
-					}
-				}
-			}
-		}
-	}
-
 	STRUCTUTILS_API uint32 GetStructCrc32(const UScriptStruct& ScriptStruct, const uint8* StructMemory, const uint32 CRC /*= 0*/)
 	{
 		FArchiveCrc32 Ar(HashCombine(CRC, PointerHash(&ScriptStruct)));
@@ -43,7 +18,7 @@ namespace UE::StructUtils
 		return Ar.GetCrc();
 	}
 
-	STRUCTUTILS_API uint32 GetStructCrc32(const FConstStructView& StructView, const uint32 CRC /*= 0*/)
+	STRUCTUTILS_API uint32 GetStructCrc32(const FConstStructView StructView, const uint32 CRC /*= 0*/)
 	{
 		if (const UScriptStruct* ScriptStruct = StructView.GetScriptStruct())
 		{
diff --git a/Engine/Plugins/Experimental/StructUtils/Source/StructUtils/Public/ChunkedStructBuffer.h b/Engine/Plugins/Experimental/StructUtils/Source/StructUtils/Public/ChunkedStructBuffer.h
--- a/Engine/Plugins/Experimental/StructUtils/Source/StructUtils/Public/ChunkedStructBuffer.h
+++ b/Engine/Plugins/Experimental/StructUtils/Source/StructUtils/Public/ChunkedStructBuffer.h
@@ -1,10 +1,10 @@
-﻿// Copyright Epic Games, Inc. All Rights Reserved.
+// Copyright Epic Games, Inc. All Rights Reserved.
 
 #pragma once
 
 #include "CoreMinimal.h"
 #include "UObject/Class.h"
-#include "InstancedStruct.h"
+#include "StructView.h"
 #include "ChunkedStructBuffer.generated.h"
 
 /**
diff --git a/Engine/Plugins/Experimental/StructUtils/Source/StructUtils/Public/InstancedStruct.h b/Engine/Plugins/Experimental/StructUtils/Source/StructUtils/Public/InstancedStruct.h
--- a/Engine/Plugins/Experimental/StructUtils/Source/StructUtils/Public/InstancedStruct.h
+++ b/Engine/Plugins/Experimental/StructUtils/Source/StructUtils/Public/InstancedStruct.h
@@ -2,198 +2,12 @@
 
 #pragma once
 
-#include "CoreMinimal.h"
-#include "UObject/Class.h"
-#include "UObject/WeakObjectPtr.h"
-#include "Templates/Casts.h"
+#include "StructUtils.h"
 
 #include "InstancedStruct.generated.h"
 
-
-
-///////////////////////////////////////////////////////////////// FConstBaseStruct /////////////////////////////////////////////////////////////////
-
-/**
- * Immutable base functionality for struct pointer handling,
- * Do not use directly, use either FInstancedStruct, FStructView or FConstStructView.
- */
-USTRUCT()
-struct STRUCTUTILS_API FConstBaseStruct
-{
-	GENERATED_BODY()
-
-public:
-
-	FConstBaseStruct() = default;
-
-	/** Copy and move should all be handled by subclasses, but it is impossible to delete them as it is a UStruct... */
-	FConstBaseStruct(const FConstBaseStruct&) { checkf(false, TEXT("Should all be handled by subclasses")); }
-	FConstBaseStruct& operator=(const FConstBaseStruct&) { checkf(false, TEXT("Should all be handled by subclasses")); return *this; }
-	FConstBaseStruct(const FConstBaseStruct&&) { checkf(false, TEXT("Should all be handled by subclasses")); }
-	FConstBaseStruct& operator=(const FConstBaseStruct&&) { checkf(false, TEXT("Should all be handled by subclasses")); return *this; }
-
-	bool operator==(const FConstBaseStruct& Other) const
-	{
-		if ((GetScriptStruct() != Other.GetScriptStruct()) || (GetMemory() != Other.GetMemory()))
-		{
-			return false;
-		}
-		return true;
-	}
-
-	bool operator!=(const FConstBaseStruct& Other) const
-	{
-		return !operator==(Other);
-	}
-
-	/** Returns const pointer to struct memory. */
-	const uint8* GetMemory() const
-	{
-		return StructMemory;
-	}
-
-	/** Returns const reference to the struct, this getter assumes that all data is valid. */
-	template<typename T>
-	const T& Get() const
-	{
-		check(StructMemory != nullptr);
-		check(ScriptStruct != nullptr);
-		check(ScriptStruct->IsChildOf(T::StaticStruct()));
-		return *((T*)StructMemory);
-	}
-
-	/** Returns const pointer to the struct, or nullptr if cast is not valid. */
-	template<typename T>
-	const T* GetPtr() const
-	{
-		if (StructMemory != nullptr && ScriptStruct && ScriptStruct->IsChildOf(T::StaticStruct()))
-		{
-			return ((T*)StructMemory);
-		}
-		return nullptr;
-	}
-
-	/** Returns struct type. */
-	const UScriptStruct* GetScriptStruct() const
-	{
-		return ScriptStruct;
-	}
-
-	/**
-	 * @return True if the struct is valid.
-	 */
-	bool IsValid() const
-	{
-		return StructMemory != nullptr && ScriptStruct != nullptr;
-	}
-
-	/**
-	 * Reset to empty.
-	 */
-	void Reset()
-	{
-		// We do not DestructStruct as we do not own it
-		StructMemory = nullptr;
-		ScriptStruct = nullptr;
-	}
-
-
-protected:
-
-	FConstBaseStruct(const UScriptStruct* InScriptStruct, const uint8* InStructMemory = nullptr)
-		: ScriptStruct(InScriptStruct)
-		, StructMemory(InStructMemory)
-	{
-	}
-
-	void SetMemory(const uint8* InStructMemory)
-	{
-		StructMemory = InStructMemory;
-	}
-
-	void SetScriptStruct(const UScriptStruct* InScriptStruct)
-	{
-		ScriptStruct = InScriptStruct;
-	}
-
-	void SetStructData(const UScriptStruct* InScriptStruct, const uint8* InStructMemory)
-	{
-		ScriptStruct = InScriptStruct;
-		StructMemory = InStructMemory;
-	}
-
-protected:
-	const UScriptStruct* ScriptStruct = nullptr;
-	const uint8* StructMemory = nullptr;
-};
-
-///////////////////////////////////////////////////////////////// FBaseStruct /////////////////////////////////////////////////////////////////
-
-/**
- * Base functionality for mutable struct pointer handling. 
- * The idea here is that it's only possible for FBaseStruct to be setup from mutable non const memory. This makes the const_cast here safe in GetMutableMemory().
- * Do not use directly, use either FInstancedStruct, or FStructView.
- * Note const FBaseStruct only makes the members of this class immutable NOT the struct data pointed at immutable.
- */
-USTRUCT()
-struct FBaseStruct : public FConstBaseStruct
-{
-	GENERATED_BODY()
-
-public:
-
-	FBaseStruct() = default;
-
-	/** Copy and move should all be handled by subclasses, but it is impossible to delete them as it is a UStruct... */
-	FBaseStruct(const FBaseStruct&) { checkf(false, TEXT("Should all be handled by subclasses")); }
-	FBaseStruct& operator=(const FBaseStruct&) { checkf(false, TEXT("Should all be handled by subclasses")); return *this; }
-	FBaseStruct(const FBaseStruct&&) { checkf(false, TEXT("Should all be handled by subclasses")); }
-	FBaseStruct& operator=(const FBaseStruct&&) { checkf(false, TEXT("Should all be handled by subclasses")); return *this; }
-
-	/** Returns a mutable pointer to struct memory. This const_cast here is safe as a FBaseStruct can only be setup from mutable non const memory. */
-	uint8* GetMutableMemory() const
-	{
-		return const_cast<uint8*>(StructMemory);
-	}
-
-	/** Returns mutable reference to the struct, this getter assumes that all data is valid. */
-	template<typename T>
-	T& GetMutable() const
-	{
-		check(StructMemory != nullptr);
-		check(ScriptStruct != nullptr);
-		check(ScriptStruct->IsChildOf(T::StaticStruct()));
-		return *((T*)GetMutableMemory());
-	}
-
-	/** Returns mutable pointer to the struct, or nullptr if cast is not valid. */
-	template<typename T>
-	T* GetMutablePtr() const
-	{
-		if (StructMemory != nullptr && ScriptStruct && ScriptStruct->IsChildOf(T::StaticStruct()))
-		{
-			return ((T*)GetMutableMemory());
-		}
-		return nullptr;
-	}
-
-protected:
-	FBaseStruct(const UScriptStruct* InScriptStruct, uint8* InStructMemory = nullptr)
-		: FConstBaseStruct(InScriptStruct, InStructMemory)
-	{}
-
-	void DestroyScriptStruct() const
-	{
-		check(StructMemory != nullptr);
-		if (ScriptStruct != nullptr)
-		{
-			ScriptStruct->DestroyStruct(GetMutableMemory());
-		}
-	}
-};
-
-///////////////////////////////////////////////////////////////// FInstancedStruct /////////////////////////////////////////////////////////////////
-
+struct FConstStructView;
+struct FConstSharedStruct;
 /**
  * FInstancedStruct works similarly as instanced UObject* property but is USTRUCTs.
  * Example:
@@ -205,56 +19,40 @@ protected:
  *	TArray<FInstancedStruct> TestArray;
  */
 USTRUCT()
-struct STRUCTUTILS_API FInstancedStruct : public FBaseStruct
+struct STRUCTUTILS_API FInstancedStruct
 {
 	GENERATED_BODY()
 
 public:
 
 	FInstancedStruct()
-		: FBaseStruct()
 	{
 	}
 
 	explicit FInstancedStruct(const UScriptStruct* InScriptStruct)
-		: FBaseStruct()
 	{
 		InitializeAs(InScriptStruct, nullptr);
 	}
 
-	FInstancedStruct(const FConstBaseStruct& InOther)
-		: FBaseStruct()
-	{
-		InitializeAs(InOther.GetScriptStruct(), InOther.GetMemory());
-	}
+	FInstancedStruct(const FConstStructView InOther);
 
 	FInstancedStruct(const FInstancedStruct& InOther)
-		: FBaseStruct()
 	{
 		InitializeAs(InOther.GetScriptStruct(), InOther.GetMemory());
 	}
 
 	FInstancedStruct(FInstancedStruct&& InOther)
-		: FBaseStruct(InOther.GetScriptStruct(), InOther.GetMutableMemory())
+		: FInstancedStruct(InOther.GetScriptStruct(), InOther.GetMutableMemory())
 	{
-		InOther.SetScriptStruct(nullptr);
-		InOther.SetMemory(nullptr);
+		InOther.SetStructData(nullptr,nullptr);
 	}
 
-
 	~FInstancedStruct()
 	{
 		Reset();
 	}
 
-	FInstancedStruct& operator=(const FConstBaseStruct& InOther)
-	{
-		if (*this != InOther)
-		{
-			InitializeAs(InOther.GetScriptStruct(), InOther.GetMemory());
-		}
-		return *this;
-	}
+	FInstancedStruct& operator=(const FConstStructView InOther);
 
 	FInstancedStruct& operator=(const FInstancedStruct& InOther)
 	{
@@ -271,11 +69,8 @@ public:
 		{
 			Reset();
 
-			SetScriptStruct(InOther.GetScriptStruct());
-			SetMemory(InOther.GetMutableMemory());
-
-			InOther.SetScriptStruct(nullptr);
-			InOther.SetMemory(nullptr);
+			SetStructData(InOther.GetScriptStruct(), InOther.GetMemory());
+			InOther.SetStructData(nullptr,nullptr);
 		}
 		return *this;
 	}
@@ -287,24 +82,23 @@ public:
 	template<typename T, typename... TArgs>
 	void InitializeAs(TArgs&&... InArgs)
 	{
-		static_assert(!TIsDerivedFrom<T, FConstBaseStruct>::IsDerived && !TIsDerivedFrom<T, FConstSharedStruct>::IsDerived, "It does not make sense to intialize a instanced struct over an other struct wrapper type");
+		UE::StructUtils::CheckStructType<T>();
 
 		Reset();
 
-		const UScriptStruct* InScriptStruct = T::StaticStruct();
-		SetScriptStruct(InScriptStruct);
-
-		const int32 RequiredSize = InScriptStruct->GetStructureSize();
-		SetMemory((uint8*)FMemory::Malloc(FMath::Max(1, RequiredSize)));
+		const UScriptStruct* Struct = T::StaticStruct();
+		const int32 RequiredSize = Struct->GetStructureSize();
+		uint8* Memory = (uint8*)FMemory::Malloc(FMath::Max(1, RequiredSize));
+		SetStructData(Struct, Memory);
 
-		new (GetMutableMemory()) T(Forward<TArgs>(InArgs)...);
+		new (Memory) T(Forward<TArgs>(InArgs)...);
 	}
 
 	/** Creates a new FInstancedStruct from templated struct type. */
 	template<typename T>
 	static FInstancedStruct Make()
 	{
-		static_assert(!TIsDerivedFrom<T, FConstBaseStruct>::IsDerived && !TIsDerivedFrom<T, FConstSharedStruct>::IsDerived, "It does not make sense to create a instanced struct over an other struct wrapper type");
+		UE::StructUtils::CheckStructType<T>();
 
 		FInstancedStruct InstancedStruct;
 		InstancedStruct.InitializeAs(T::StaticStruct(), nullptr);
@@ -315,7 +109,7 @@ public:
 	template<typename T>
 	static FInstancedStruct Make(const T& Struct)
 	{
-		static_assert(!TIsDerivedFrom<T, FConstBaseStruct>::IsDerived && !TIsDerivedFrom<T, FConstSharedStruct>::IsDerived, "It does not make sense to create a instanced struct over an other struct wrapper type");
+		UE::StructUtils::CheckStructType<T>();
 
 		FInstancedStruct InstancedStruct;
 		InstancedStruct.InitializeAs(T::StaticStruct(), reinterpret_cast<const uint8*>(&Struct));
@@ -326,16 +120,13 @@ public:
 	template<typename T, typename... TArgs>
 	static inline FInstancedStruct Make(TArgs&&... InArgs)
 	{
-		static_assert(!TIsDerivedFrom<T, FConstBaseStruct>::IsDerived && !TIsDerivedFrom<T, FConstSharedStruct>::IsDerived, "It does not make sense to create a instanced struct over an other struct wrapper type");
+		UE::StructUtils::CheckStructType<T>();
 
 		FInstancedStruct InstancedStruct;
 		InstancedStruct.InitializeAs<T>(Forward<TArgs>(InArgs)...);
 		return InstancedStruct;
 	}
 
-	/** Reset to empty. */
-	void Reset();
-
 	/** For StructOpsTypeTraits */
 	bool Serialize(FArchive& Ar);
 	bool Identical(const FInstancedStruct* Other, uint32 PortFlags) const;
@@ -343,148 +134,30 @@ public:
 	bool ExportTextItem(FString& ValueStr, FInstancedStruct const& DefaultValue, class UObject* Parent, int32 PortFlags, class UObject* ExportRootScope) const;
 	bool ImportTextItem(const TCHAR*& Buffer, int32 PortFlags, UObject* Parent, FOutputDevice* ErrorText, FArchive* InSerializingArchive = nullptr);
 
-private:
-	/** Initializes for new struct type (does nothing if same type) and returns mutable struct. */
-	UScriptStruct* ReinitializeAs(const UScriptStruct* InScriptStruct);
-};
-
-template<>
-struct TStructOpsTypeTraits<FInstancedStruct> : public TStructOpsTypeTraitsBase2<FInstancedStruct>
-{
-	enum
-	{
-		WithSerializer = true,
-		WithIdentical = true,
-		WithExportTextItem = true,
-		WithImportTextItem = true,
-		WithAddStructReferencedObjects = true,
-	};
-};
-
-///////////////////////////////////////////////////////////////// FStructSharedMemory /////////////////////////////////////////////////////////////////
-
-/**
- * Holds the information and memory about a UStruct and that is the actual part that is shared across all the FConstSharedStruct/FSharedStruct
- * 
- * The size of the allocation for this structure should always includes not only the need size for it members but also the size required to hold the
- * structure describe by SciprtStruct. This is how we can avoid 2 pointer referencing(cache misses). Look at the Create() method to understand more.
- */
-struct STRUCTUTILS_API FStructSharedMemory : public TSharedFromThis<FStructSharedMemory>
-{
-	~FStructSharedMemory()
-	{
-		ScriptStruct.DestroyStruct(GetMemory());
-	}
-
-	static TSharedPtr<FStructSharedMemory> Create(const UScriptStruct& InScriptStruct, const uint8* InStructMemory = nullptr)
-	{
-		const int32 RequiredSize = sizeof(FStructSharedMemory) + InScriptStruct.GetStructureSize();
-		// Code analysis is unable to understand correctly what we are doing here, so disabling the warning C6386: Buffer overrun while writing to...
-		CA_SUPPRESS( 6386 )
-		FStructSharedMemory* StructMemory = new(new uint8[RequiredSize]) FStructSharedMemory(InScriptStruct, InStructMemory);
-		return MakeShareable(StructMemory);
-	}
-
-	/** Returns pointer to struct memory. */
-	uint8* GetMemory() const
-	{
-		return (uint8*)StructMemory;
-	}
-
 	/** Returns struct type. */
-	const UScriptStruct& GetScriptStruct() const
+	const UScriptStruct* GetScriptStruct() const
 	{
 		return ScriptStruct;
 	}
 
-private:
-	FStructSharedMemory(const UScriptStruct& InScriptStruct, const uint8* InStructMemory = nullptr)
-		: ScriptStruct(InScriptStruct)
-	{
-		if (InStructMemory)
-		{
-			ScriptStruct.CopyScriptStruct(StructMemory, InStructMemory);
-		}
-	}
-
-	const UScriptStruct& ScriptStruct;
-
-	// The required memory size for the struct represented by the UScriptStruct must be allocated right after this object into big enough preallocated buffer, 
-	// Check Create() method for more information.
-	uint8 StructMemory[0];
-};
-
-///////////////////////////////////////////////////////////////// FConstSharedStruct /////////////////////////////////////////////////////////////////
-
-/**
- * FConstSharedStruct is the same as the FSharedStruct but restrict the API to return const struct type. 
- * 
- * See FSharedStruct for more information.
- */
-USTRUCT()
-struct STRUCTUTILS_API FConstSharedStruct
-{
-	GENERATED_BODY();
-
-	FConstSharedStruct()
-	{}
-
-	FConstSharedStruct(const FConstSharedStruct& InOther)
-	{
-		StructMemoryPtr = InOther.StructMemoryPtr;
-	}
-
-	FConstSharedStruct(const FConstSharedStruct&& InOther)
-	{
-		StructMemoryPtr = InOther.StructMemoryPtr;
-	}
-
-	FConstSharedStruct& operator=(const FConstSharedStruct& InOther)
-	{
-		if (this != &InOther)
-		{
-			StructMemoryPtr = InOther.StructMemoryPtr;
-		}
-		return *this;
-	}
-	FConstSharedStruct& operator=(FConstSharedStruct&& InOther)
-	{
-		if (this != &InOther)
-		{
-			StructMemoryPtr = InOther.StructMemoryPtr;
-		}
-		return *this;
-	}
-
-	bool operator==(const FConstSharedStruct& Other) const
-	{
-		if ((GetScriptStruct() != Other.GetScriptStruct()) || (GetMemory() != Other.GetMemory()))
-		{
-			return false;
-		}
-		return true;
-	}
-
-	bool operator!=(const FConstSharedStruct& Other) const
-	{
-		return !operator==(Other);
-	}
-
 	/** Returns const pointer to struct memory. */
 	const uint8* GetMemory() const
 	{
-		return StructMemoryPtr ? StructMemoryPtr.Get()->GetMemory() : nullptr;
+		return StructMemory;
 	}
 
+	/** Reset to empty. */
+	void Reset();
+
 	/** Returns const reference to the struct, this getter assumes that all data is valid. */
 	template<typename T>
 	const T& Get() const
 	{
 		const uint8* Memory = GetMemory();
-		const UScriptStruct* ScriptStruct = GetScriptStruct();
+		const UScriptStruct* Struct = GetScriptStruct();
 		check(Memory != nullptr);
-		check(ScriptStruct != nullptr);
-		check(ScriptStruct->IsChildOf(T::StaticStruct()));
+		check(Struct != nullptr);
+		check(Struct->IsChildOf(T::StaticStruct()));
 		return *((T*)Memory);
 	}
 
@@ -493,123 +166,19 @@ struct STRUCTUTILS_API FConstSharedStruct
 	const T* GetPtr() const
 	{
 		const uint8* Memory = GetMemory();
-		const UScriptStruct* ScriptStruct = GetScriptStruct();
-		if (Memory != nullptr && ScriptStruct && ScriptStruct->IsChildOf(T::StaticStruct()))
+		const UScriptStruct* Struct = GetScriptStruct();
+		if (Memory != nullptr && Struct && Struct->IsChildOf(T::StaticStruct()))
 		{
 			return ((T*)Memory);
 		}
 		return nullptr;
 	}
 
-	/** Returns struct type. */
-	const UScriptStruct* GetScriptStruct() const
-	{
-		return StructMemoryPtr ? &(StructMemoryPtr.Get()->GetScriptStruct()) : nullptr;
-	}
-
-	/**
-	 * @return True if the struct is valid.
-	 */
-	bool IsValid() const
-	{
-		return GetMemory() != nullptr && GetScriptStruct() != nullptr;
-	}
-
-	/**
-	 * Reset to empty.
-	 */
-	void Reset()
-	{
-		StructMemoryPtr.Reset();
-	}
-
-	/** For StructOpsTypeTraits */
-	bool Identical(const FConstSharedStruct* Other, uint32 PortFlags) const;
-	void AddStructReferencedObjects(class FReferenceCollector& Collector);
-
-protected:
-
-	TSharedPtr<const FStructSharedMemory> StructMemoryPtr;
-};
-
-template<>
-struct TStructOpsTypeTraits<FConstSharedStruct> : public TStructOpsTypeTraitsBase2<FConstSharedStruct>
-{
-	enum
-	{
-		WithIdentical = true,
-		WithAddStructReferencedObjects = true,
-	};
-};
-
-///////////////////////////////////////////////////////////////// FSharedStruct /////////////////////////////////////////////////////////////////
-
-/**
- * FSharedStruct works similarly as a TSharedPtr<FInstancedStruct> but removes the double pointer indirection that would create.
- * (One pointer for the FInstancedStruct and one pointer for the struct memory it is wrapping).
- * Also note that because of its implementation, it is not possible for now to go from a struct reference or struct view back to a shared struct. 
- * 
- * This struct type is also convertible to a FStructView and is the preferable way of passing it as a parameter just as the FInstancedStruct.
- * If the calling code would like to keep a shared pointer to the struct, you may pass the FSharedStruct as a parameter but it is recommended to pass it as 
- * a "const FSharedStruct&" to limit the unnecessary recounting.
- * 
- */
-USTRUCT()
-struct STRUCTUTILS_API FSharedStruct : public FConstSharedStruct
-{
-	GENERATED_BODY();
-
-	FSharedStruct()
-	{
-	}
-
-	explicit FSharedStruct(const UScriptStruct* InScriptStruct)
-	{
-		InitializeAs(InScriptStruct, nullptr);
-	}
-
-	FSharedStruct(const FConstBaseStruct& InOther)
-	{
-		InitializeAs(InOther.GetScriptStruct(), InOther.GetMemory());
-	}
-
-	FSharedStruct(const FSharedStruct& InOther)
-		: FConstSharedStruct(InOther)
-	{
-	}
-
-	FSharedStruct(FSharedStruct&& InOther)
-		: FConstSharedStruct(InOther)
-	{
-	}
-
-	~FSharedStruct()
-	{
-		Reset();
-	}
-
-	FSharedStruct& operator=(const FConstBaseStruct& InOther)
-	{
-		InitializeAs(InOther.GetScriptStruct(), InOther.GetMemory());
-		return *this;
-	}
-
-	FSharedStruct& operator=(const FSharedStruct& InOther)
-	{
-		FConstSharedStruct::operator=(InOther);
-		return *this;
-	}
-
-	FSharedStruct& operator=(FSharedStruct&& InOther)
-	{
-		FConstSharedStruct::operator=(InOther);
-		return *this;
-	}
-
-	/** Returns a mutable pointer to struct memory. This const_cast here is safe as a FSharedStruct can only be setup from mutable non const memory. */
+	/** Returns a mutable pointer to struct memory. This const_cast here is safe as a ClassName can only be setup from mutable non const memory. */
 	uint8* GetMutableMemory() const
 	{
-		return const_cast<uint8*>(GetMemory());
+		const uint8* Memory = GetMemory();
+		return const_cast<uint8*>(Memory);
 	}
 
 	/** Returns mutable reference to the struct, this getter assumes that all data is valid. */
@@ -617,10 +186,10 @@ struct STRUCTUTILS_API FSharedStruct : public FConstSharedStruct
 	T& GetMutable() const
 	{
 		uint8* Memory = GetMutableMemory();
-		const UScriptStruct* ScriptStruct = GetScriptStruct();
+		const UScriptStruct* Struct = GetScriptStruct();
 		check(Memory != nullptr);
-		check(ScriptStruct != nullptr);
-		check(ScriptStruct->IsChildOf(T::StaticStruct()));
+		check(Struct != nullptr);
+		check(Struct->IsChildOf(T::StaticStruct()));
 		return *((T*)Memory);
 	}
 
@@ -629,186 +198,80 @@ struct STRUCTUTILS_API FSharedStruct : public FConstSharedStruct
 	T* GetMutablePtr() const
 	{
 		uint8* Memory = GetMutableMemory();
-		const UScriptStruct* ScriptStruct = GetScriptStruct();
-		if (Memory != nullptr && ScriptStruct && ScriptStruct->IsChildOf(T::StaticStruct()))
+		const UScriptStruct* Struct = GetScriptStruct();
+		if (Memory != nullptr && Struct && Struct->IsChildOf(T::StaticStruct()))
 		{
 			return ((T*)Memory);
 		}
 		return nullptr;
 	}
 
-	/** Initializes from struct type and optional data. */
-	void InitializeAs(const UScriptStruct* InScriptStruct, const uint8* InStructMemory = nullptr)
+	/** Returns True if the struct is valid.*/
+	bool IsValid() const
 	{
-		Reset();
-		if (InScriptStruct)
+		return GetMemory() != nullptr && GetScriptStruct() != nullptr;
+	}
+
+	/** Comparison operators. Note: it does not compare the internal structure itself*/
+	template <typename OtherType>
+	bool operator==(const OtherType& Other) const
+	{
+		if ((GetScriptStruct() != Other.GetScriptStruct()) || (GetMemory() != Other.GetMemory()))
 		{
-			StructMemoryPtr = FStructSharedMemory::Create(*InScriptStruct, InStructMemory);
+			return false;
 		}
+		return true;
 	}
 
-	/** Initializes from struct type and emplace construct. */
-	template<typename T, typename... TArgs>
-	void InitializeAs(TArgs&&... InArgs)
+	template <typename OtherType>
+	bool operator!=(const OtherType& Other) const
 	{
-		static_assert(!TIsDerivedFrom<T, FConstBaseStruct>::IsDerived && !TIsDerivedFrom<T, FConstSharedStruct>::IsDerived, "It does not make sense to initialize a shared struct over an other struct wrapper type");
-
-		Reset();
-		StructMemoryPtr = FStructSharedMemory::Create(*T::StaticStruct());
-		new (GetMutableMemory()) T(Forward<TArgs>(InArgs)...);
+		return !operator==(Other);
 	}
 
-	/** Creates a new FSharedStruct from templated struct type. */
-	template<typename T>
-	static FSharedStruct Make()
-	{
-		static_assert(!TIsDerivedFrom<T, FConstBaseStruct>::IsDerived && !TIsDerivedFrom<T, FConstSharedStruct>::IsDerived, "It does not make sense to create a shared struct over an other struct wrapper type");
+protected:
 
-		FSharedStruct SharedStruct;
-		SharedStruct.InitializeAs(T::StaticStruct(), nullptr);
-		return SharedStruct;
-	}
+	/** Initializes for new struct type (does nothing if same type) and returns mutable struct. */
+	UScriptStruct* ReinitializeAs(const UScriptStruct* InScriptStruct);
 
-	/** Creates a new FSharedStruct from templated struct. */
-	template<typename T>
-	static FSharedStruct Make(const T& Struct)
+	void DestroyScriptStruct() const
 	{
-		static_assert(!TIsDerivedFrom<T, FConstBaseStruct>::IsDerived && !TIsDerivedFrom<T, FConstSharedStruct>::IsDerived, "It does not make sense to create a shared struct over an other struct wrapper type");
-
-		FSharedStruct SharedStruct;
-		SharedStruct.InitializeAs(T::StaticStruct(), reinterpret_cast<const uint8*>(&Struct));
-		return SharedStruct;
+		check(StructMemory != nullptr);
+		if (ScriptStruct != nullptr)
+		{
+			ScriptStruct->DestroyStruct(GetMutableMemory());
+		}
 	}
 
-	/** Creates a new FSharedStruct from the templated type and forward all arguments to constructor. */
-	template<typename T, typename... TArgs>
-	static inline FSharedStruct Make(TArgs&&... InArgs)
+	FInstancedStruct(const UScriptStruct* InScriptStruct, const uint8* InStructMemory)
+		: ScriptStruct(InScriptStruct)
+		, StructMemory(InStructMemory)
+	{}
+	void ResetStructData()
 	{
-		static_assert(!TIsDerivedFrom<T, FConstBaseStruct>::IsDerived && !TIsDerivedFrom<T, FConstSharedStruct>::IsDerived, "It does not make sense to create a shared struct over an other struct wrapper type");
-
-		FSharedStruct SharedStruct;
-		SharedStruct.InitializeAs<T>(Forward<TArgs>(InArgs)...);
-		return SharedStruct;
+		StructMemory = nullptr;
+		ScriptStruct = nullptr;
 	}
+	void SetStructData(const UScriptStruct* InScriptStruct, const uint8* InStructMemory)
+	{
+		ScriptStruct = InScriptStruct;
+		StructMemory = InStructMemory;
+	}
+
+
+	const UScriptStruct* ScriptStruct = nullptr; \
+	const uint8* StructMemory = nullptr;
 };
 
 template<>
-struct TStructOpsTypeTraits<FSharedStruct> : public TStructOpsTypeTraitsBase2<FSharedStruct>
+struct TStructOpsTypeTraits<FInstancedStruct> : public TStructOpsTypeTraitsBase2<FInstancedStruct>
 {
 	enum
 	{
+		WithSerializer = true,
 		WithIdentical = true,
+		WithExportTextItem = true,
+		WithImportTextItem = true,
 		WithAddStructReferencedObjects = true,
 	};
 };
-
-///////////////////////////////////////////////////////////////// FConstStructView /////////////////////////////////////////////////////////////////
-
-/**
- * FConstStructView is "typed" struct pointer, it contains const pointer to struct plus UScriptStruct pointer.
- * FConstStructView does not own the memory and will not free it when out of scope.
- * It should be only used to pass struct pointer in a limited scope, or when the user controls the lifetime of the struct being stored.
- * E.g. instead of passing ref or pointer to a FInstancedStruct, you should use FConstStructView or FStructView to pass around a view to the contents.
- * FConstStructView is passed by value.
- * FConstStructView is similar to FStructOnScope, but FConstStructView is a view only (FStructOnScope can either own the memory or be a view)
- */
-struct STRUCTUTILS_API FConstStructView : public FConstBaseStruct
-{
-
-public:
-
-	FConstStructView() = default;
-
-	FConstStructView(const UScriptStruct* InScriptStruct, const uint8* InStructMemory = nullptr)
-		: FConstBaseStruct(InScriptStruct, InStructMemory)
-	{}
-
-	FConstStructView(const FConstBaseStruct& ConstBaseStruct)
-		: FConstBaseStruct(ConstBaseStruct.GetScriptStruct(), ConstBaseStruct.GetMemory())
-	{}
-
-	FConstStructView(const FConstSharedStruct& SharedStruct)
-		: FConstBaseStruct(SharedStruct.GetScriptStruct(), SharedStruct.GetMemory())
-	{}
-
-	FConstStructView(const FConstStructView& Other)
-		: FConstBaseStruct(Other.GetScriptStruct(), Other.GetMemory())
-	{}
-
-	FConstStructView& operator=(const FConstBaseStruct& ConstBaseStruct)
-	{
-		SetStructData(ConstBaseStruct.GetScriptStruct(), ConstBaseStruct.GetMemory());
-		return *this;
-	}
-
-	FConstStructView& operator=(const FConstStructView& Other)
-	{
-		SetStructData(Other.GetScriptStruct(), Other.GetMemory());
-		return *this;
-	}
-
-	/** Creates a new FStructView from the templated struct */
-	template<typename T>
-	static FConstStructView Make(const T& Struct)
-	{
-		static_assert(!TIsDerivedFrom<T, FConstBaseStruct>::IsDerived && !TIsDerivedFrom<T, FConstSharedStruct>::IsDerived, "It does not make sense to create a const struct view over an other struct wrapper type");
-		return FConstStructView(T::StaticStruct(), reinterpret_cast<const uint8*>(&Struct));
-	}
-};
-
-///////////////////////////////////////////////////////////////// FStructView /////////////////////////////////////////////////////////////////
-
-/**
- * FStructView is "typed" struct pointer, it contains pointer to struct plus UScriptStruct pointer.
- * FStructView does not own the memory and will not free it when out of scope.
- * It should be only used to pass struct pointer in a limited scope, or when the user controls the lifetime of the struct being stored.
- * E.g. instead of passing ref or pointer to a FInstancedStruct, you should use FConstStructView or FStructView to pass around a view to the contents.
- * FStructView is passed by value.
- * FStructView is similar to FStructOnScope, but FStructView is a view only (FStructOnScope can either own the memory or be a view)
- */
-struct STRUCTUTILS_API FStructView : public FBaseStruct
-{
-
-public:
-
-	FStructView()
-		: FBaseStruct()
-	{
-	}
-
-	FStructView(const UScriptStruct* InScriptStruct, uint8* InStructMemory = nullptr)
-		: FBaseStruct(InScriptStruct, InStructMemory)
-	{}
-
-	FStructView(const FBaseStruct& BaseStruct)
-		: FBaseStruct(BaseStruct.GetScriptStruct(), BaseStruct.GetMutableMemory())
-	{}
-
-	FStructView(const FSharedStruct& SharedStruct)
-		: FBaseStruct(SharedStruct.GetScriptStruct(), SharedStruct.GetMutableMemory())
-	{}
-
-	FStructView(const FStructView& Other)
-		: FBaseStruct(Other.GetScriptStruct(), Other.GetMutableMemory())
-	{}
-
-	FStructView& operator=(const FBaseStruct& BaseStruct)
-	{
-		SetStructData(BaseStruct.GetScriptStruct(), BaseStruct.GetMutableMemory());
-		return *this;
-	}
-
-	FStructView& operator=(const FStructView& Other)
-	{
-		SetStructData(Other.GetScriptStruct(), Other.GetMutableMemory());
-		return *this;
-	}
-
-	/** Creates a new FStructView from the templated struct. Note its not safe to make InStruct const ref as the original object may have been declared const */
-	template<typename T>
-	static FStructView Make(T& InStruct)
-	{
-		static_assert(!TIsDerivedFrom<T, FConstBaseStruct>::IsDerived && !TIsDerivedFrom<T, FConstSharedStruct>::IsDerived, "It does not make sense to create a struct view over an other struct wrapper type");
-		return FStructView(T::StaticStruct(), reinterpret_cast<uint8*>(&InStruct));
-	}
-};
\ No newline at end of file
diff --git a/Engine/Plugins/Experimental/StructUtils/Source/StructUtils/Public/InstancedStructStream.h b/Engine/Plugins/Experimental/StructUtils/Source/StructUtils/Public/InstancedStructStream.h
--- a/Engine/Plugins/Experimental/StructUtils/Source/StructUtils/Public/InstancedStructStream.h
+++ b/Engine/Plugins/Experimental/StructUtils/Source/StructUtils/Public/InstancedStructStream.h
@@ -4,7 +4,7 @@
 
 #include "CoreMinimal.h"
 #include "UObject/Class.h"
-#include "InstancedStruct.h"
+#include "StructView.h"
 #include "InstancedStructStream.generated.h"
 
 /**
diff --git a/Engine/Plugins/Experimental/StructUtils/Source/StructUtils/Public/SharedStruct.h b/Engine/Plugins/Experimental/StructUtils/Source/StructUtils/Public/SharedStruct.h
new file mode 100644
--- /dev/null
+++ b/Engine/Plugins/Experimental/StructUtils/Source/StructUtils/Public/SharedStruct.h
@@ -0,0 +1,311 @@
+// Copyright Epic Games, Inc. All Rights Reserved.
+
+#pragma once
+
+#include "StructUtils.h"
+
+#include "SharedStruct.generated.h"
+
+struct FInstancedStruct;
+struct FConstStructView;
+
+///////////////////////////////////////////////////////////////// FStructSharedMemory /////////////////////////////////////////////////////////////////
+
+/**
+ * Holds the information and memory about a UStruct and that is the actual part that is shared across all the FConstSharedStruct/FSharedStruct
+ * 
+ * The size of the allocation for this structure should always includes not only the need size for it members but also the size required to hold the
+ * structure describe by SciprtStruct. This is how we can avoid 2 pointer referencing(cache misses). Look at the Create() method to understand more.
+ */
+struct STRUCTUTILS_API FStructSharedMemory : public TSharedFromThis<FStructSharedMemory>
+{
+	~FStructSharedMemory()
+	{
+		ScriptStruct.DestroyStruct(GetMemory());
+	}
+
+	static TSharedPtr<FStructSharedMemory> Create(const UScriptStruct& InScriptStruct, const uint8* InStructMemory = nullptr)
+	{
+		const int32 RequiredSize = sizeof(FStructSharedMemory) + InScriptStruct.GetStructureSize();
+		// Code analysis is unable to understand correctly what we are doing here, so disabling the warning C6386: Buffer overrun while writing to...
+		CA_SUPPRESS( 6386 )
+		FStructSharedMemory* StructMemory = new(new uint8[RequiredSize]) FStructSharedMemory(InScriptStruct, InStructMemory);
+		return MakeShareable(StructMemory);
+	}
+
+	/** Returns pointer to struct memory. */
+	uint8* GetMemory() const
+	{
+		return (uint8*)StructMemory;
+	}
+
+	/** Returns struct type. */
+	const UScriptStruct& GetScriptStruct() const
+	{
+		return ScriptStruct;
+	}
+
+private:
+	FStructSharedMemory(const UScriptStruct& InScriptStruct, const uint8* InStructMemory = nullptr)
+		: ScriptStruct(InScriptStruct)
+	{
+		if (InStructMemory)
+		{
+			ScriptStruct.CopyScriptStruct(StructMemory, InStructMemory);
+		}
+	}
+
+	const UScriptStruct& ScriptStruct;
+
+	// The required memory size for the struct represented by the UScriptStruct must be allocated right after this object into big enough preallocated buffer, 
+	// Check Create() method for more information.
+	uint8 StructMemory[0];
+};
+
+///////////////////////////////////////////////////////////////// FConstSharedStruct /////////////////////////////////////////////////////////////////
+
+/**
+ * FConstSharedStruct is the same as the FSharedStruct but restrict the API to return const struct type. 
+ * 
+ * See FSharedStruct for more information.
+ */
+USTRUCT()
+struct STRUCTUTILS_API FConstSharedStruct
+{
+	GENERATED_BODY();
+
+	FConstSharedStruct()
+	{}
+
+	/** Copy constructors */
+	FConstSharedStruct(const FConstSharedStruct& InOther) = default;
+	FConstSharedStruct(FConstSharedStruct&& InOther) = default;
+
+	/** Assignment operators */
+	FConstSharedStruct& operator=(const FConstSharedStruct& InOther) = default;
+	FConstSharedStruct& operator=(FConstSharedStruct&& InOther) = default;
+
+	/** For StructOpsTypeTraits */
+	bool Identical(const FConstSharedStruct* Other, uint32 PortFlags) const;
+	void AddStructReferencedObjects(class FReferenceCollector& Collector);
+
+	/** Returns struct type. */
+	const UScriptStruct* GetScriptStruct() const
+	{
+		return StructMemoryPtr ? &(StructMemoryPtr.Get()->GetScriptStruct()) : nullptr;
+	}
+
+	/** Returns const pointer to struct memory. */
+	const uint8* GetMemory() const
+	{
+		return StructMemoryPtr ? StructMemoryPtr.Get()->GetMemory() : nullptr;
+	}
+
+	/** Reset to empty. */
+	void Reset()
+	{
+		StructMemoryPtr.Reset();
+	}
+
+	/** Returns const reference to the struct, this getter assumes that all data is valid. */
+	template<typename T>
+	const T& Get() const
+	{
+		const uint8* Memory = GetMemory();
+		const UScriptStruct* Struct = GetScriptStruct();
+		check(Memory != nullptr);
+		check(Struct != nullptr);
+		check(Struct->IsChildOf(T::StaticStruct()));
+		return *((T*)Memory);
+	}
+
+	/** Returns const pointer to the struct, or nullptr if cast is not valid. */
+	template<typename T>
+	const T* GetPtr() const
+	{
+		const uint8* Memory = GetMemory();
+		const UScriptStruct* Struct = GetScriptStruct();
+		if (Memory != nullptr && Struct && Struct->IsChildOf(T::StaticStruct()))
+		{
+			return ((T*)Memory);
+		}
+		return nullptr;
+	}
+
+	/** Returns True if the struct is valid.*/
+	bool IsValid() const
+	{
+		return GetMemory() != nullptr && GetScriptStruct() != nullptr;
+	}
+
+	/** Comparison operators. Note: it does not compare the internal structure itself*/
+	template <typename OtherType>
+	bool operator==(const OtherType& Other) const
+	{
+		if ((GetScriptStruct() != Other.GetScriptStruct()) || (GetMemory() != Other.GetMemory()))
+		{
+			return false;
+		}
+		return true;
+	}
+
+	template <typename OtherType>
+	bool operator!=(const OtherType& Other) const
+	{
+		return !operator==(Other);
+	}
+
+protected:
+
+	TSharedPtr<const FStructSharedMemory> StructMemoryPtr;
+};
+
+template<>
+struct TStructOpsTypeTraits<FConstSharedStruct> : public TStructOpsTypeTraitsBase2<FConstSharedStruct>
+{
+	enum
+	{
+		WithIdentical = true,
+		WithAddStructReferencedObjects = true,
+	};
+};
+
+///////////////////////////////////////////////////////////////// FSharedStruct /////////////////////////////////////////////////////////////////
+
+/**
+ * FSharedStruct works similarly as a TSharedPtr<FInstancedStruct> but removes the double pointer indirection that would create.
+ * (One pointer for the FInstancedStruct and one pointer for the struct memory it is wrapping).
+ * Also note that because of its implementation, it is not possible for now to go from a struct reference or struct view back to a shared struct. 
+ * 
+ * This struct type is also convertible to a FStructView and is the preferable way of passing it as a parameter just as the FInstancedStruct.
+ * If the calling code would like to keep a shared pointer to the struct, you may pass the FSharedStruct as a parameter but it is recommended to pass it as 
+ * a "const FSharedStruct&" to limit the unnecessary recounting.
+ * 
+ */
+USTRUCT()
+struct STRUCTUTILS_API FSharedStruct : public FConstSharedStruct
+{
+	GENERATED_BODY();
+
+	FSharedStruct()
+	{
+	}
+
+	explicit FSharedStruct(const UScriptStruct* InScriptStruct)
+	{
+		InitializeAs(InScriptStruct, nullptr);
+	}
+
+	FSharedStruct(const FConstStructView InOther);
+
+	~FSharedStruct()
+	{
+		Reset();
+	}
+
+	/** Copy constructors */
+	FSharedStruct(const FSharedStruct& InOther) = default;
+	FSharedStruct(FSharedStruct&& InOther) = default;
+
+	/** Assignment operators */
+	FSharedStruct& operator=(const FSharedStruct& InOther) = default;
+	FSharedStruct& operator=(FSharedStruct&& InOther) = default;
+
+	FSharedStruct& operator=(const FConstStructView InOther);
+
+	/** Initializes from struct type and optional data. */
+	void InitializeAs(const UScriptStruct* InScriptStruct, const uint8* InStructMemory = nullptr)
+	{
+		Reset();
+		if (InScriptStruct)
+		{
+			StructMemoryPtr = FStructSharedMemory::Create(*InScriptStruct, InStructMemory);
+		}
+	}
+
+	/** Initializes from struct type and emplace construct. */
+	template<typename T, typename... TArgs>
+	void InitializeAs(TArgs&&... InArgs)
+	{
+		UE::StructUtils::CheckStructType<T>();
+
+		Reset();
+		StructMemoryPtr = FStructSharedMemory::Create(*T::StaticStruct());
+		new (GetMutableMemory()) T(Forward<TArgs>(InArgs)...);
+	}
+
+	/** Creates a new FSharedStruct from templated struct type. */
+	template<typename T>
+	static FSharedStruct Make()
+	{
+		UE::StructUtils::CheckStructType<T>();
+
+		FSharedStruct SharedStruct;
+		SharedStruct.InitializeAs(T::StaticStruct(), nullptr);
+		return SharedStruct;
+	}
+
+	/** Creates a new FSharedStruct from templated struct. */
+	template<typename T>
+	static FSharedStruct Make(const T& Struct)
+	{
+		UE::StructUtils::CheckStructType<T>();
+
+		FSharedStruct SharedStruct;
+		SharedStruct.InitializeAs(T::StaticStruct(), reinterpret_cast<const uint8*>(&Struct));
+		return SharedStruct;
+	}
+
+	/** Creates a new FSharedStruct from the templated type and forward all arguments to constructor. */
+	template<typename T, typename... TArgs>
+	static inline FSharedStruct Make(TArgs&&... InArgs)
+	{
+		UE::StructUtils::CheckStructType<T>();
+
+		FSharedStruct SharedStruct;
+		SharedStruct.InitializeAs<T>(Forward<TArgs>(InArgs)...);
+		return SharedStruct;
+	}
+
+	/** Returns a mutable pointer to struct memory. This const_cast here is safe as a ClassName can only be setup from mutable non const memory. */
+	uint8* GetMutableMemory() const
+	{
+		const uint8* Memory = GetMemory();
+		return const_cast<uint8*>(Memory);
+	}
+
+	/** Returns mutable reference to the struct, this getter assumes that all data is valid. */
+	template<typename T>
+	T& GetMutable() const
+	{
+		uint8* Memory = GetMutableMemory();
+		const UScriptStruct* Struct = GetScriptStruct();
+		check(Memory != nullptr);
+		check(Struct != nullptr);
+		check(Struct->IsChildOf(T::StaticStruct()));
+		return *((T*)Memory);
+	}
+
+	/** Returns mutable pointer to the struct, or nullptr if cast is not valid. */
+	template<typename T>
+	T* GetMutablePtr() const
+	{
+		uint8* Memory = GetMutableMemory();
+		const UScriptStruct* Struct = GetScriptStruct();
+		if (Memory != nullptr && Struct && Struct->IsChildOf(T::StaticStruct()))
+		{
+			return ((T*)Memory);
+		}
+		return nullptr;
+	}
+};
+
+template<>
+struct TStructOpsTypeTraits<FSharedStruct> : public TStructOpsTypeTraitsBase2<FSharedStruct>
+{
+	enum
+	{
+		WithIdentical = true,
+		WithAddStructReferencedObjects = true,
+	};
+};
\ No newline at end of file
diff --git a/Engine/Plugins/Experimental/StructUtils/Source/StructUtils/Public/StructUtils.h b/Engine/Plugins/Experimental/StructUtils/Source/StructUtils/Public/StructUtils.h
new file mode 100644
--- /dev/null
+++ b/Engine/Plugins/Experimental/StructUtils/Source/StructUtils/Public/StructUtils.h
@@ -0,0 +1,16 @@
+// Copyright Epic Games, Inc. All Rights Reserved.
+
+#pragma once
+
+#include "Templates/UnrealTypeTraits.h"
+
+namespace UE::StructUtils
+{
+	template <typename T>
+	void CheckStructType()
+	{
+		static_assert(!TIsDerivedFrom<T, struct FInstancedStruct>::IsDerived &&
+					  !TIsDerivedFrom<T, struct FConstStructView>::IsDerived &&
+					  !TIsDerivedFrom<T, struct FConstSharedStruct>::IsDerived, "It does not make sense to create a instanced struct over an other struct wrapper type");
+	}
+}
diff --git a/Engine/Plugins/Experimental/StructUtils/Source/StructUtils/Public/StructUtilsTypes.h b/Engine/Plugins/Experimental/StructUtils/Source/StructUtils/Public/StructUtilsTypes.h
--- a/Engine/Plugins/Experimental/StructUtils/Source/StructUtils/Public/StructUtilsTypes.h
+++ b/Engine/Plugins/Experimental/StructUtils/Source/StructUtils/Public/StructUtilsTypes.h
@@ -3,7 +3,7 @@
 #pragma once
 
 #include "UObject/Class.h"
-#include "InstancedStruct.h"
+#include "StructView.h"
 
 #ifndef WITH_STRUCTUTILS_DEBUG
 #define WITH_STRUCTUTILS_DEBUG (!(UE_BUILD_SHIPPING || UE_BUILD_SHIPPING_WITH_EDITOR || UE_BUILD_TEST) && 1)
@@ -14,11 +14,9 @@ class FReferenceCollector;
 
 namespace UE::StructUtils
 {
-	extern STRUCTUTILS_API void AddStructReferencedObjects(const FConstStructView& StructView, class FReferenceCollector& Collector);
-
 	extern STRUCTUTILS_API uint32 GetStructCrc32(const UScriptStruct& ScriptStruct, const uint8* StructMemory, const uint32 CRC = 0);
 
-	extern STRUCTUTILS_API uint32 GetStructCrc32(const FConstStructView& StructView, const uint32 CRC = 0);
+	extern STRUCTUTILS_API uint32 GetStructCrc32(const FConstStructView StructView, const uint32 CRC = 0);
 }
 
 /* Predicate useful to find a struct of a specific type in an container */
@@ -26,9 +24,9 @@ struct FStructTypeEqualOperator
 {
 	const UScriptStruct* TypePtr;
 	FStructTypeEqualOperator(const UScriptStruct* InTypePtr) : TypePtr(InTypePtr) {}
-	FStructTypeEqualOperator(const FConstStructView& InRef) : TypePtr(InRef.GetScriptStruct()) {}
+	FStructTypeEqualOperator(const FConstStructView InRef) : TypePtr(InRef.GetScriptStruct()) {}
 
-	bool operator()(const FConstStructView& Other) const { return Other.GetScriptStruct() == TypePtr; }
+	bool operator()(const FConstStructView Other) const { return Other.GetScriptStruct() == TypePtr; }
 };
 
 struct FScriptStructSortOperator
@@ -43,7 +41,7 @@ struct FScriptStructSortOperator
 
 struct FStructTypeSortOperator
 {
-	bool operator()(const FConstStructView& A, const FConstStructView& B) const
+	bool operator()(const FConstStructView A, const FConstStructView B) const
 	{
 		const UScriptStruct* AScriptStruct = A.GetScriptStruct();
 		const UScriptStruct* BScriptStruct = B.GetScriptStruct();
diff --git a/Engine/Plugins/Experimental/StructUtils/Source/StructUtils/Public/StructView.h b/Engine/Plugins/Experimental/StructUtils/Source/StructUtils/Public/StructView.h
new file mode 100644
--- /dev/null
+++ b/Engine/Plugins/Experimental/StructUtils/Source/StructUtils/Public/StructView.h
@@ -0,0 +1,218 @@
+// Copyright Epic Games, Inc. All Rights Reserved.
+
+#pragma once
+
+#include "StructUtils.h"
+#include "InstancedStruct.h"
+#include "SharedStruct.h"
+
+///////////////////////////////////////////////////////////////// FConstStructView /////////////////////////////////////////////////////////////////
+
+/**
+ * FConstStructView is "typed" struct pointer, it contains const pointer to struct plus UScriptStruct pointer.
+ * FConstStructView does not own the memory and will not free it when out of scope.
+ * It should be only used to pass struct pointer in a limited scope, or when the user controls the lifetime of the struct being stored.
+ * E.g. instead of passing ref or pointer to a FInstancedStruct, you should use FConstStructView or FStructView to pass around a view to the contents.
+ * FConstStructView is passed by value.
+ * FConstStructView is similar to FStructOnScope, but FConstStructView is a view only (FStructOnScope can either own the memory or be a view)
+ */
+struct STRUCTUTILS_API FConstStructView
+{
+public:
+
+	FConstStructView() = default;
+
+	FConstStructView(const UScriptStruct* InScriptStruct, const uint8* InStructMemory = nullptr)
+		: ScriptStruct(InScriptStruct)
+		, StructMemory(InStructMemory)
+	{}
+
+	FConstStructView(const FInstancedStruct& InstancedStruct)
+		: FConstStructView(InstancedStruct.GetScriptStruct(), InstancedStruct.GetMemory())
+
+	{}
+
+	FConstStructView(const FConstSharedStruct& SharedStruct)
+		: FConstStructView(SharedStruct.GetScriptStruct(), SharedStruct.GetMemory())
+
+	{}
+
+	/** Copy constructors */
+	FConstStructView(const FConstStructView& Other) = default;
+	FConstStructView(FConstStructView&& Other) = default;
+
+	/** Assignment operators */
+	FConstStructView& operator=(const FConstStructView& Other) = default;
+	FConstStructView& operator=(FConstStructView&& Other) = default;
+
+	/** Creates a new FStructView from the templated struct */
+	template<typename T>
+	static FConstStructView Make(const T& Struct)
+	{
+		UE::StructUtils::CheckStructType<T>();
+		return FConstStructView(T::StaticStruct(), reinterpret_cast<const uint8*>(&Struct));
+	}
+
+	/** Returns struct type. */
+	const UScriptStruct* GetScriptStruct() const
+	{
+		return ScriptStruct;
+	}
+
+	/** Returns const pointer to struct memory. */
+	const uint8* GetMemory() const
+	{
+		return StructMemory;
+	}
+
+	/** Reset to empty. */
+	void Reset()
+	{
+		ResetStructData();
+	}
+
+	/** Returns const reference to the struct, this getter assumes that all data is valid. */
+	template<typename T>
+	const T& Get() const
+	{
+		const uint8* Memory = GetMemory();
+		const UScriptStruct* Struct = GetScriptStruct();
+		check(Memory != nullptr);
+		check(Struct != nullptr);
+		check(Struct->IsChildOf(T::StaticStruct()));
+		return *((T*)Memory);
+	}
+
+	/** Returns const pointer to the struct, or nullptr if cast is not valid. */
+	template<typename T>
+	const T* GetPtr() const
+	{
+		const uint8* Memory = GetMemory();
+		const UScriptStruct* Struct = GetScriptStruct();
+		if (Memory != nullptr && Struct && Struct->IsChildOf(T::StaticStruct()))
+		{
+			return ((T*)Memory);
+		}
+		return nullptr;
+	}
+
+	/** Returns True if the struct is valid.*/
+	bool IsValid() const
+	{
+		return GetMemory() != nullptr && GetScriptStruct() != nullptr;
+	}
+
+	/** Comparison operators. Note: it does not compare the internal structure itself */
+	template <typename OtherType>
+	bool operator==(const OtherType& Other) const
+	{
+		if ((GetScriptStruct() != Other.GetScriptStruct()) || (GetMemory() != Other.GetMemory()))
+		{
+			return false;
+		}
+		return true;
+	}
+
+	template <typename OtherType>
+	bool operator!=(const OtherType& Other) const
+	{
+		return !operator==(Other);
+	}
+
+protected:
+
+	void ResetStructData()
+	{
+		StructMemory = nullptr;
+		ScriptStruct = nullptr;
+	}
+	void SetStructData(const UScriptStruct* InScriptStruct, const uint8* InStructMemory)
+	{
+		ScriptStruct = InScriptStruct;
+		StructMemory = InStructMemory;
+	}
+
+	const UScriptStruct* ScriptStruct = nullptr;
+	const uint8* StructMemory = nullptr;
+};
+
+
+///////////////////////////////////////////////////////////////// FStructView /////////////////////////////////////////////////////////////////
+
+/**
+ * FStructView is "typed" struct pointer, it contains pointer to struct plus UScriptStruct pointer.
+ * FStructView does not own the memory and will not free it when out of scope.
+ * It should be only used to pass struct pointer in a limited scope, or when the user controls the lifetime of the struct being stored.
+ * E.g. instead of passing ref or pointer to a FInstancedStruct, you should use FConstStructView or FStructView to pass around a view to the contents.
+ * FStructView is passed by value.
+ * FStructView is similar to FStructOnScope, but FStructView is a view only (FStructOnScope can either own the memory or be a view)
+ */
+struct STRUCTUTILS_API FStructView : public FConstStructView
+{
+public:
+
+	FStructView()
+		: FConstStructView()
+	{
+	}
+
+	FStructView(const UScriptStruct* InScriptStruct, uint8* InStructMemory = nullptr)
+		: FConstStructView(InScriptStruct, InStructMemory)
+	{}
+
+	FStructView(const FInstancedStruct& InstancedStruct)
+		: FConstStructView(InstancedStruct.GetScriptStruct(), InstancedStruct.GetMutableMemory())
+	{}
+
+	FStructView(const FSharedStruct& SharedStruct)
+		: FConstStructView(SharedStruct.GetScriptStruct(), SharedStruct.GetMutableMemory())
+	{}
+
+	/** Copy constructors */
+	FStructView(const FStructView& Other) = default;
+	FStructView(FStructView&& Other) = default;
+
+	/** Assignment operators */
+	FStructView& operator=(const FStructView& Other) = default;
+	FStructView& operator=(FStructView&& Other) = default;
+
+	/** Creates a new FStructView from the templated struct. Note its not safe to make InStruct const ref as the original object may have been declared const */
+	template<typename T>
+	static FStructView Make(T& InStruct)
+	{
+		UE::StructUtils::CheckStructType<T>();
+		return FStructView(T::StaticStruct(), reinterpret_cast<uint8*>(&InStruct));
+	}
+
+	/** Returns a mutable pointer to struct memory. This const_cast here is safe as a ClassName can only be setup from mutable non const memory. */
+	uint8* GetMutableMemory() const
+	{
+		const uint8* Memory = GetMemory();
+		return const_cast<uint8*>(Memory);
+	}
+
+	/** Returns mutable reference to the struct, this getter assumes that all data is valid. */
+	template<typename T>
+	T& GetMutable() const
+	{
+		uint8* Memory = GetMutableMemory();
+		const UScriptStruct* Struct = GetScriptStruct();
+		check(Memory != nullptr);
+		check(Struct != nullptr);
+		check(Struct->IsChildOf(T::StaticStruct()));
+		return *((T*)Memory);
+	}
+
+	/** Returns mutable pointer to the struct, or nullptr if cast is not valid. */
+	template<typename T>
+	T* GetMutablePtr() const
+	{
+		uint8* Memory = GetMutableMemory();
+		const UScriptStruct* Struct = GetScriptStruct();
+		if (Memory != nullptr && Struct && Struct->IsChildOf(T::StaticStruct()))
+		{
+			return ((T*)Memory);
+		}
+		return nullptr;
+	}
+};
\ No newline at end of file
diff --git a/Engine/Plugins/Runtime/MassEntity/Source/MassEntity/Private/MassArchetypeData.cpp b/Engine/Plugins/Runtime/MassEntity/Source/MassEntity/Private/MassArchetypeData.cpp
--- a/Engine/Plugins/Runtime/MassEntity/Source/MassEntity/Private/MassArchetypeData.cpp
+++ b/Engine/Plugins/Runtime/MassEntity/Source/MassEntity/Private/MassArchetypeData.cpp
@@ -70,7 +70,7 @@ void FMassArchetypeData::Initialize(const FMassArchetypeCompositionDescriptor& I
 	for (const UScriptStruct* ChunkFragmentType : ChunkFragmentList)
 	{
 		check(ChunkFragmentType);
-		ChunkFragmentsTemplate.Add(FConstStructView(ChunkFragmentType));
+		ChunkFragmentsTemplate.Emplace(ChunkFragmentType);
 	}
 
 	// Share fragments
diff --git a/Engine/Plugins/Runtime/StateTree/Source/StateTreeModule/Public/StateTreePropertyBindings.h b/Engine/Plugins/Runtime/StateTree/Source/StateTreeModule/Public/StateTreePropertyBindings.h
--- a/Engine/Plugins/Runtime/StateTree/Source/StateTreeModule/Public/StateTreePropertyBindings.h
+++ b/Engine/Plugins/Runtime/StateTree/Source/StateTreeModule/Public/StateTreePropertyBindings.h
@@ -4,7 +4,7 @@
 
 #include "CoreMinimal.h"
 #include "Misc/Guid.h"
-#include "InstancedStruct.h"
+#include "StructView.h"
 #include "StateTreePropertyBindings.generated.h"
 
 class FProperty;
