Remove unecessary backslash

#rnx

diff --git a/Engine/Plugins/Experimental/StructUtils/Source/StructUtils/Public/InstancedStruct.h b/Engine/Plugins/Experimental/StructUtils/Source/StructUtils/Public/InstancedStruct.h
--- a/Engine/Plugins/Experimental/StructUtils/Source/StructUtils/Public/InstancedStruct.h
+++ b/Engine/Plugins/Experimental/StructUtils/Source/StructUtils/Public/InstancedStruct.h
@@ -259,7 +259,7 @@ protected:
 	}
 
 
-	const UScriptStruct* ScriptStruct = nullptr; \
+	const UScriptStruct* ScriptStruct = nullptr;
 	const uint8* StructMemory = nullptr;
 };
 
