Removed TSharedFromThis in the FStructMemeory, it was not used.

#rnx

diff --git a/Engine/Plugins/Experimental/StructUtils/Source/StructUtils/Public/SharedStruct.h b/Engine/Plugins/Experimental/StructUtils/Source/StructUtils/Public/SharedStruct.h
--- a/Engine/Plugins/Experimental/StructUtils/Source/StructUtils/Public/SharedStruct.h
+++ b/Engine/Plugins/Experimental/StructUtils/Source/StructUtils/Public/SharedStruct.h
@@ -17,7 +17,7 @@ struct FConstStructView;
  * The size of the allocation for this structure should always includes not only the need size for it members but also the size required to hold the
  * structure describe by SciprtStruct. This is how we can avoid 2 pointer referencing(cache misses). Look at the Create() method to understand more.
  */
-struct STRUCTUTILS_API FStructSharedMemory : public TSharedFromThis<FStructSharedMemory>
+struct STRUCTUTILS_API FStructSharedMemory
 {
 	~FStructSharedMemory()
 	{
