Extended TScriptStructTypeBitSet with fast-tracks for types known at compilation time

diff --git a/Engine/Plugins/Experimental/StructUtils/Source/StructUtils/Public/ScriptStructTypeBitSet.h b/Engine/Plugins/Experimental/StructUtils/Source/StructUtils/Public/ScriptStructTypeBitSet.h
--- a/Engine/Plugins/Experimental/StructUtils/Source/StructUtils/Public/ScriptStructTypeBitSet.h
+++ b/Engine/Plugins/Experimental/StructUtils/Source/StructUtils/Public/ScriptStructTypeBitSet.h
@@ -190,6 +190,41 @@ private:
 			}
 			return Hash;
 		}
+
+		void AddAtIndex(const int32 Index)
+		{
+			PadToNum(Index + 1, false);
+			SetBitNoCheck(Index, true);
+		}
+
+		void RemoveAtIndex(const int32 Index)
+		{
+			check(Index >= 0);
+			if (Index < Num())
+			{
+				SetBitNoCheck(Index, false);
+			}
+			// else, it's already not present
+		}
+
+		bool Contains(const int32 Index) const
+		{
+			check(Index >= 0);
+			return (Index < Num()) && (*this)[Index];
+		}
+
+	protected:
+		/**
+		 * duplication of TBitArray::SetBitNoCheck needed since it's private but it's the performant way of setting bits
+		 * when we know the index is valid.
+		 * @todo ask Core team about exposing that
+		 */
+		void SetBitNoCheck(const int32 Index, const bool Value)
+		{
+			uint32& Word = GetData()[Index / NumBitsPerDWORD];
+			const uint32 BitOffset = (Index % NumBitsPerDWORD);
+			Word = (Word & ~(1 << BitOffset)) | (((uint32)Value) << BitOffset);
+		}
 	};
 
 	static FStructTracker StructTracker;
@@ -248,18 +283,39 @@ private:
 
 public:
 
+	static int32 CreateTypeIndex(const UScriptStruct& InStructType)
+	{
+#if WITH_STRUCTUTILS_DEBUG
+		ensureMsgf(InStructType.IsChildOf(TBaseStruct::StaticStruct())
+			, TEXT("Creating index for '%s' while it doesn't derive from the expected struct type %s")
+			, *InStructType.GetPathName(), *TBaseStruct::StaticStruct()->GetName());
+#endif // WITH_STRUCTUTILS_DEBUG
+
+		return StructTracker.FindOrAddStructTypeIndex(InStructType);
+	}
+
+	template<typename T>
+	static int32 GetTypeIndex()
+	{
+		static_assert(TIsDerivedFrom<T, TBaseStruct>::IsDerived, "Given struct type doesn't match the expected base struct type.");
+		static const int32 TypeIndex = CreateTypeIndex(*T::StaticStruct());
+		return TypeIndex;
+	}
+
 	template<typename T>
 	FORCEINLINE void Add()
 	{
 		static_assert(TIsDerivedFrom<T, TBaseStruct>::IsDerived, "Given struct type doesn't match the expected base struct type.");
-		Add(*T::StaticStruct());
+		const int32 StructTypeIndex = GetTypeIndex<T>();
+		StructTypesBitArray.AddAtIndex(StructTypeIndex);
 	}
 
 	template<typename T>
 	FORCEINLINE void Remove()
 	{
 		static_assert(TIsDerivedFrom<T, TBaseStruct>::IsDerived, "Given struct type doesn't match the expected base struct type.");
-		Remove(*T::StaticStruct());
+		const int32 StructTypeIndex = GetTypeIndex<T>();
+		StructTypesBitArray.RemoveAtIndex(StructTypeIndex);
 	}
 
 	FORCEINLINE void Remove(const TScriptStructTypeBitSet<TBaseStruct>& Other)
@@ -271,7 +327,8 @@ public:
 	FORCEINLINE bool Contains() const
 	{
 		static_assert(TIsDerivedFrom<T, TBaseStruct>::IsDerived, "Given struct type doesn't match the expected base struct type.");
-		return Contains(*T::StaticStruct());
+		const int32 StructTypeIndex = GetTypeIndex<T>();
+		return StructTypesBitArray.Contains(StructTypeIndex);
 	}
 
 	void Add(const UScriptStruct& InStructType)
@@ -283,8 +340,7 @@ public:
 #endif // WITH_STRUCTUTILS_DEBUG
 
 		const int32 StructTypeIndex = StructTracker.FindOrAddStructTypeIndex(InStructType);
-		StructTypesBitArray.PadToNum(StructTypeIndex + 1, false);
-		StructTypesBitArray[StructTypeIndex] = true;
+		StructTypesBitArray.AddAtIndex(StructTypeIndex);
 	}
 
 	void Remove(const UScriptStruct& InStructType)
@@ -296,11 +352,7 @@ public:
 #endif // WITH_STRUCTUTILS_DEBUG
 
 		const int32 StructTypeIndex = StructTracker.FindOrAddStructTypeIndex(InStructType);
-		if (StructTypeIndex < StructTypesBitArray.Num())
-		{
-			StructTypesBitArray[StructTypeIndex] = false;
-		}
-		// else, if it's not in StructTypesBitArray it's already not present
+		StructTypesBitArray.RemoveAtIndex(StructTypeIndex);
 	}
 
 	void Reset() { StructTypesBitArray.Reset(); }
@@ -314,7 +366,7 @@ public:
 #endif // WITH_STRUCTUTILS_DEBUG
 
 		const int32 StructTypeIndex = StructTracker.FindOrAddStructTypeIndex(InStructType);
-		return (StructTypeIndex < StructTypesBitArray.Num()) && StructTypesBitArray[StructTypeIndex];
+		return StructTypesBitArray.Contains(StructTypeIndex);
 	}
 
 	FORCEINLINE TScriptStructTypeBitSet operator+(const TScriptStructTypeBitSet& Other) const
