[FInstancedStruct] filter out details view when metadata "Hidden" is used

diff --git a/Engine/Plugins/Experimental/StructUtils/Source/StructUtilsEditor/Private/Customizations/InstancedStructDetails.cpp b/Engine/Plugins/Experimental/StructUtils/Source/StructUtilsEditor/Private/Customizations/InstancedStructDetails.cpp
--- a/Engine/Plugins/Experimental/StructUtils/Source/StructUtilsEditor/Private/Customizations/InstancedStructDetails.cpp
+++ b/Engine/Plugins/Experimental/StructUtils/Source/StructUtilsEditor/Private/Customizations/InstancedStructDetails.cpp
@@ -46,6 +46,11 @@ public:
 			return bAllowBaseStruct;
 		}
 
+		if (InStruct->HasMetaData(TEXT("Hidden")))
+		{
+			return false;
+		}
+
 		// Query the native struct to see if it has the correct parent type (if any)
 		return !BaseStruct || InStruct->IsChildOf(BaseStruct);
 	}
