SharedStruct: Call InitializeStruct() when initializing shared memory
- Fixes crash on more complex types (i.e. when the struct has array)

#rnx

diff --git a/Engine/Plugins/Experimental/StructUtils/Source/StructUtils/Public/SharedStruct.h b/Engine/Plugins/Experimental/StructUtils/Source/StructUtils/Public/SharedStruct.h
--- a/Engine/Plugins/Experimental/StructUtils/Source/StructUtils/Public/SharedStruct.h
+++ b/Engine/Plugins/Experimental/StructUtils/Source/StructUtils/Public/SharedStruct.h
@@ -49,6 +49,8 @@ private:
 	FStructSharedMemory(const UScriptStruct& InScriptStruct, const uint8* InStructMemory = nullptr)
 		: ScriptStruct(InScriptStruct)
 	{
+		ScriptStruct.InitializeStruct(StructMemory);
+		
 		if (InStructMemory)
 		{
 			ScriptStruct.CopyScriptStruct(StructMemory, InStructMemory);
