SharedStruct: Ensure FStructSharedMemory struct memory is aligned according to the wrapped struct's alignment
- Fixes crash on types with > 8 alignment e.g: any types using FTransform etc

#rnx

diff --git a/Engine/Plugins/Experimental/StructUtils/Source/StructUtils/Public/SharedStruct.h b/Engine/Plugins/Experimental/StructUtils/Source/StructUtils/Public/SharedStruct.h
--- a/Engine/Plugins/Experimental/StructUtils/Source/StructUtils/Public/SharedStruct.h
+++ b/Engine/Plugins/Experimental/StructUtils/Source/StructUtils/Public/SharedStruct.h
@@ -24,19 +24,29 @@ struct STRUCTUTILS_API FStructSharedMemory
 		ScriptStruct.DestroyStruct(GetMemory());
 	}
 
+	struct FStructSharedMemoryDeleter
+	{
+		FORCEINLINE void operator()(FStructSharedMemory* StructSharedMemory) const
+		{
+			FMemory::Free(StructSharedMemory);
+		}
+	};
+
 	static TSharedPtr<FStructSharedMemory> Create(const UScriptStruct& InScriptStruct, const uint8* InStructMemory = nullptr)
 	{
-		const int32 RequiredSize = sizeof(FStructSharedMemory) + InScriptStruct.GetStructureSize();
+		// Align RequiredSize to InScriptStruct's alignment to effectively add padding in between ScriptStruct and
+		// StructMemory. GetMemory will then round &StructMemory up past this 'padding' to the nearest aligned address.
+		const int32 RequiredSize = Align(sizeof(FStructSharedMemory), InScriptStruct.GetMinAlignment()) + InScriptStruct.GetStructureSize();
 		// Code analysis is unable to understand correctly what we are doing here, so disabling the warning C6386: Buffer overrun while writing to...
 		CA_SUPPRESS( 6386 )
-		FStructSharedMemory* StructMemory = new(new uint8[RequiredSize]) FStructSharedMemory(InScriptStruct, InStructMemory);
-		return MakeShareable(StructMemory);
+		FStructSharedMemory* StructMemory = new(FMemory::Malloc(RequiredSize, InScriptStruct.GetMinAlignment())) FStructSharedMemory(InScriptStruct, InStructMemory);
+		return MakeShareable(StructMemory, FStructSharedMemoryDeleter());
 	}
 
-	/** Returns pointer to struct memory. */
+	/** Returns pointer to aligned struct memory. */
 	uint8* GetMemory() const
 	{
-		return (uint8*)StructMemory;
+		return Align((uint8*)StructMemory, ScriptStruct.GetMinAlignment());
 	}
 
 	/** Returns struct type. */
@@ -49,11 +59,11 @@ private:
 	FStructSharedMemory(const UScriptStruct& InScriptStruct, const uint8* InStructMemory = nullptr)
 		: ScriptStruct(InScriptStruct)
 	{
-		ScriptStruct.InitializeStruct(StructMemory);
+		ScriptStruct.InitializeStruct(GetMemory());
 		
 		if (InStructMemory)
 		{
-			ScriptStruct.CopyScriptStruct(StructMemory, InStructMemory);
+			ScriptStruct.CopyScriptStruct(GetMemory(), InStructMemory);
 		}
 	}
 
