Copying FStructVariant features to FInstancedStruct

 - Copied localization support to FInstancedStruct.
 - Copied Import/ExportText support for null structs to FInstancedStruct.
 - Copied Preload support to FInstancedStruct.
 - Copied Identical logic (deep-compare) to FInstancedStruct.
   - The old logic always did a shallow-compare, which for FInstancedStruct==FInstancedStruct would always be false and meant the property didn't behave as a value type.
   - FInstancedStruct==AnythingElse still only performs a shallow-compare, as this may be comparing against a view type.
 - Added support for FInstancedStruct to in-place upgrade (via SerializeFromMismatchedTag) from FStructVariant.

diff --git a/Engine/Plugins/Experimental/StructUtils/Source/StructUtils/Private/InstancedStruct.cpp b/Engine/Plugins/Experimental/StructUtils/Source/StructUtils/Private/InstancedStruct.cpp
--- a/Engine/Plugins/Experimental/StructUtils/Source/StructUtils/Private/InstancedStruct.cpp
+++ b/Engine/Plugins/Experimental/StructUtils/Source/StructUtils/Private/InstancedStruct.cpp
@@ -2,9 +2,48 @@
 #include "InstancedStruct.h"
 #include "StructView.h"
 #include "SharedStruct.h"
+#include "UObject/StructVariant.h"
+#include "Serialization/PropertyLocalizationDataGathering.h"
+
+namespace UE::StructUtils::Private
+{
+#if WITH_EDITORONLY_DATA
+void GatherForLocalization(const FString& PathToParent, const UScriptStruct* Struct, const void* StructData, const void* DefaultStructData, FPropertyLocalizationDataGatherer& PropertyLocalizationDataGatherer, const EPropertyLocalizationGathererTextFlags GatherTextFlags)
+{
+	const FInstancedStruct* ThisInstance = static_cast<const FInstancedStruct*>(StructData);
+	const FInstancedStruct* DefaultInstance = static_cast<const FInstancedStruct*>(DefaultStructData);
+
+	PropertyLocalizationDataGatherer.GatherLocalizationDataFromStruct(PathToParent, Struct, StructData, DefaultStructData, GatherTextFlags);
+
+	if (const UScriptStruct* StructTypePtr = ThisInstance->GetScriptStruct())
+	{
+		PropertyLocalizationDataGatherer.GatherLocalizationDataFromStructWithCallbacks(PathToParent + TEXT(".StructInstance"), StructTypePtr, ThisInstance->GetMemory(), DefaultInstance ? DefaultInstance->GetMemory() : nullptr, GatherTextFlags);
+	}
+}
+#endif // WITH_EDITORONLY_DATA
+
+void RegisterForLocalization()
+{
+#if WITH_EDITORONLY_DATA
+	{ static const FAutoRegisterLocalizationDataGatheringCallback AutomaticRegistrationOfLocalizationGatherer(TBaseStructure<FInstancedStruct>::Get(), &GatherForLocalization); }
+#endif
+}
+}
+
+FInstancedStruct::FInstancedStruct()
+{
+	UE::StructUtils::Private::RegisterForLocalization();
+}
+
+FInstancedStruct::FInstancedStruct(const UScriptStruct* InScriptStruct)
+{
+	UE::StructUtils::Private::RegisterForLocalization();
+	InitializeAs(InScriptStruct, nullptr);
+}
 
 FInstancedStruct::FInstancedStruct(const FConstStructView InOther)
 {
+	UE::StructUtils::Private::RegisterForLocalization();
 	InitializeAs(InOther.GetScriptStruct(), InOther.GetMemory());
 }
 
@@ -108,6 +147,10 @@ bool FInstancedStruct::Serialize(FArchive& Ar)
 		// UScriptStruct type
 		UScriptStruct* NewNonConstStruct = nullptr;
 		Ar << NewNonConstStruct;
+		if (NewNonConstStruct)
+		{
+			Ar.Preload(NewNonConstStruct);
+		}
 		NonConstStruct = ReinitializeAs(NewNonConstStruct);
 
 		// Size of the serialized memory
@@ -162,44 +205,80 @@ bool FInstancedStruct::Serialize(FArchive& Ar)
 
 bool FInstancedStruct::ExportTextItem(FString& ValueStr, FInstancedStruct const& DefaultValue, class UObject* Parent, int32 PortFlags, class UObject* ExportRootScope) const
 {
-	UScriptStruct* NonConstStruct = const_cast<UScriptStruct*>(GetScriptStruct());
-	if (NonConstStruct == nullptr || GetMemory() == nullptr)
+	if (const UScriptStruct* StructTypePtr = GetScriptStruct())
 	{
-		return false;
+		ValueStr += StructTypePtr->GetPathName();
+		StructTypePtr->ExportText(ValueStr, GetMemory(), StructTypePtr == DefaultValue.GetScriptStruct() ? DefaultValue.GetMemory() : nullptr, Parent, PortFlags, ExportRootScope);
+	}
+	else
+	{
+		ValueStr += TEXT("None");
 	}
-	
-	ValueStr += NonConstStruct->GetPathName();
-
-	NonConstStruct->ExportText(ValueStr, GetMemory(), GetMemory(), nullptr, EPropertyPortFlags::PPF_None, nullptr);
 	return true;
 }
 
 bool FInstancedStruct::ImportTextItem(const TCHAR*& Buffer, int32 PortFlags, UObject* Parent, FOutputDevice* ErrorText, FArchive* InSerializingArchive /*= nullptr*/)
 {
-	FString StructPath;
-	const TCHAR* NewBuffer = FPropertyHelpers::ReadToken(Buffer, /* out */ StructPath, /* dotted names */ true);
-	if (NewBuffer == nullptr)
+	FNameBuilder StructPathName;
+	if (const TCHAR* Result = FPropertyHelpers::ReadToken(Buffer, StructPathName, /*bDottedNames*/true))
 	{
-		return false;
+		Buffer = Result;
 	}
-
-	Buffer = NewBuffer;
-		
-	UScriptStruct* NonConstStruct = ReinitializeAs(FindObject<UScriptStruct>(nullptr, *StructPath, false));
-	if (NonConstStruct == nullptr)
+	else
 	{
 		return false;
 	}
 
-	const TCHAR* Result = NonConstStruct->ImportText(Buffer, GetMutableMemory(), nullptr, EPropertyPortFlags::PPF_None, nullptr, NonConstStruct->GetName());
-	if (Result != nullptr)
+	if (StructPathName.Len() == 0 || FCString::Stricmp(StructPathName.ToString(), TEXT("None")) == 0)
 	{
-		Buffer = Result;
+		ReinitializeAs(nullptr);
+	}
+	else
+	{
+		UScriptStruct* StructTypePtr = LoadObject<UScriptStruct>(nullptr, StructPathName.ToString());
+		if (!StructTypePtr)
+		{
+			return false;
+		}
+
+		ReinitializeAs(StructTypePtr);
+		if (const TCHAR* Result = StructTypePtr->ImportText(Buffer, GetMutableMemory(), Parent, PortFlags, ErrorText, [StructTypePtr]() { return StructTypePtr->GetName(); }))
+		{
+			Buffer = Result;
+		}
+		else
+		{
+			return false;
+		}
 	}
 
 	return true;
 }
 
+bool FInstancedStruct::SerializeFromMismatchedTag(const FPropertyTag& Tag, FStructuredArchive::FSlot Slot)
+{
+	static const FName NAME_StructVariant = "StructVariant";
+	if (Tag.Type == NAME_StructProperty && Tag.StructName == NAME_StructVariant)
+	{
+		// TODO: When deleting FStructVariant, migrate its Serialize loading code here
+		FStructVariant StructVariant;
+		StructVariant.Serialize(Slot);
+
+		InitializeAs(StructVariant.GetStructType(), static_cast<uint8*>(StructVariant.GetStructInstance()));
+		return true;
+	}
+
+	return false;
+}
+
+void FInstancedStruct::GetPreloadDependencies(TArray<UObject*>& OutDeps)
+{
+	if (UScriptStruct* NonConstStruct = const_cast<UScriptStruct*>(GetScriptStruct()))
+	{
+		OutDeps.Add(NonConstStruct);
+	}
+}
+
 UScriptStruct* FInstancedStruct::ReinitializeAs(const UScriptStruct* InScriptStruct)
 {
 	UScriptStruct* NonConstStruct = const_cast<UScriptStruct*>(GetScriptStruct());
@@ -215,8 +294,18 @@ UScriptStruct* FInstancedStruct::ReinitializeAs(const UScriptStruct* InScriptStr
 
 bool FInstancedStruct::Identical(const FInstancedStruct* Other, uint32 PortFlags) const
 {
-	// Only empty is considered equal
-	return Other != nullptr && GetMemory() == nullptr && Other->GetMemory() == nullptr && GetScriptStruct() == nullptr && Other->GetScriptStruct() == nullptr;
+	const UScriptStruct* StructTypePtr = GetScriptStruct();
+	if (!Other || StructTypePtr != Other->GetScriptStruct())
+	{
+		return false;
+	}
+
+	if (StructTypePtr)
+	{
+		return StructTypePtr->CompareScriptStruct(GetMemory(), Other->GetMemory(), PortFlags);
+	}
+
+	return true;
 }
 
 void FInstancedStruct::AddStructReferencedObjects(class FReferenceCollector& Collector)
@@ -225,4 +314,4 @@ void FInstancedStruct::AddStructReferencedObjects(class FReferenceCollector& Col
 	{
 		Collector.AddReferencedObjects(Struct, GetMutableMemory());
 	}
-}
\ No newline at end of file
+}
diff --git a/Engine/Plugins/Experimental/StructUtils/Source/StructUtils/Public/InstancedStruct.h b/Engine/Plugins/Experimental/StructUtils/Source/StructUtils/Public/InstancedStruct.h
--- a/Engine/Plugins/Experimental/StructUtils/Source/StructUtils/Public/InstancedStruct.h
+++ b/Engine/Plugins/Experimental/StructUtils/Source/StructUtils/Public/InstancedStruct.h
@@ -25,14 +25,9 @@ struct STRUCTUTILS_API FInstancedStruct
 
 public:
 
-	FInstancedStruct()
-	{
-	}
+	FInstancedStruct();
 
-	explicit FInstancedStruct(const UScriptStruct* InScriptStruct)
-	{
-		InitializeAs(InScriptStruct, nullptr);
-	}
+	explicit FInstancedStruct(const UScriptStruct* InScriptStruct);
 
 	FInstancedStruct(const FConstStructView InOther);
 
@@ -133,6 +128,8 @@ public:
 	void AddStructReferencedObjects(class FReferenceCollector& Collector);
 	bool ExportTextItem(FString& ValueStr, FInstancedStruct const& DefaultValue, class UObject* Parent, int32 PortFlags, class UObject* ExportRootScope) const;
 	bool ImportTextItem(const TCHAR*& Buffer, int32 PortFlags, UObject* Parent, FOutputDevice* ErrorText, FArchive* InSerializingArchive = nullptr);
+	bool SerializeFromMismatchedTag(const FPropertyTag& Tag, FStructuredArchive::FSlot Slot);
+	void GetPreloadDependencies(TArray<UObject*>& OutDeps);
 
 	/** Returns struct type. */
 	const UScriptStruct* GetScriptStruct() const
@@ -212,6 +209,17 @@ public:
 		return GetMemory() != nullptr && GetScriptStruct() != nullptr;
 	}
 
+	/** Comparison operators. Deep compares the struct instance when identical. */
+	bool operator==(const FInstancedStruct& Other) const
+	{
+		return Identical(&Other, PPF_None);
+	}
+
+	bool operator!=(const FInstancedStruct& Other) const
+	{
+		return !Identical(&Other, PPF_None);
+	}
+
 	/** Comparison operators. Note: it does not compare the internal structure itself*/
 	template <typename OtherType>
 	bool operator==(const OtherType& Other) const
@@ -273,5 +281,7 @@ struct TStructOpsTypeTraits<FInstancedStruct> : public TStructOpsTypeTraitsBase2
 		WithExportTextItem = true,
 		WithImportTextItem = true,
 		WithAddStructReferencedObjects = true,
+		WithStructuredSerializeFromMismatchedTag = true,
+		WithGetPreloadDependencies = true,
 	};
 };
