Improved FInstancedStruct property editing stability

FInstancedStruct now edits a copy of the data (like FStructVariant did) to avoid issues where the detail panel would try and access invalid data. This copy is migrated back to the source whenever a property on the struct is edited, and the source is synced to the copy either every 0.1s or immediately on undo/redo or when the source is edited via its property handle.

The custom revert logic has been removed now that FInstancedStruct implements Identical correctly (as a deep-compare for a value type), and the customization now reads the meta-data from the correct property (to handle containers of FInstancedStruct) as well as handling the "HideViewOptions", "ShowTreeView", and CPF_NoClear meta-data that the standard struct picker supports.

diff --git a/Engine/Plugins/Experimental/StructUtils/Source/StructUtilsEditor/Private/Customizations/InstancedStructDetails.cpp b/Engine/Plugins/Experimental/StructUtils/Source/StructUtilsEditor/Private/Customizations/InstancedStructDetails.cpp
--- a/Engine/Plugins/Experimental/StructUtils/Source/StructUtilsEditor/Private/Customizations/InstancedStructDetails.cpp
+++ b/Engine/Plugins/Experimental/StructUtils/Source/StructUtilsEditor/Private/Customizations/InstancedStructDetails.cpp
@@ -62,6 +62,36 @@ public:
 	}
 };
 
+namespace UE::StructUtils::Private
+{
+
+const UScriptStruct* GetCommonScriptStruct(TSharedPtr<IPropertyHandle> StructProperty)
+{
+	const UScriptStruct* CommonStructType = nullptr;
+
+	StructProperty->EnumerateConstRawData([&CommonStructType](const void* RawData, const int32 /*DataIndex*/, const int32 /*NumDatas*/)
+	{
+		if (RawData)
+		{
+			const FInstancedStruct* InstancedStruct = static_cast<const FInstancedStruct*>(RawData);
+
+			const UScriptStruct* StructTypePtr = InstancedStruct->GetScriptStruct();
+			if (CommonStructType && CommonStructType != StructTypePtr)
+			{
+				// Multiple struct types on the sources - show nothing set
+				CommonStructType = nullptr;
+				return false;
+			}
+			CommonStructType = StructTypePtr;
+		}
+
+		return true;
+	});
+
+	return CommonStructType;
+}
+
+}
 
 ////////////////////////////////////
 
@@ -76,38 +106,108 @@ FInstancedStructDataDetails::FInstancedStructDataDetails(TSharedPtr<IPropertyHan
 	StructProperty = InStructProperty;
 }
 
-FInstancedStruct* FInstancedStructDataDetails::GetSingleInstancedStruct()
+void FInstancedStructDataDetails::OnStructValuePreChange()
+{
+	// Forward the change event to the real struct handle
+	if (StructProperty && StructProperty->IsValidHandle())
+	{
+		StructProperty->NotifyPreChange();
+	}
+}
+
+void FInstancedStructDataDetails::OnStructValuePostChange()
+{
+	// Copy the modified struct data back to the source instances
+	SyncEditableInstanceToSource();
+
+	// Forward the change event to the real struct handle
+	if (StructProperty && StructProperty->IsValidHandle())
+	{
+		TGuardValue<bool> HandlingStructValuePostChangeGuard(bIsHandlingStructValuePostChange, true);
+
+		StructProperty->NotifyPostChange(EPropertyChangeType::ValueSet);
+		StructProperty->NotifyFinishedChangingProperties();
+	}
+}
+
+void FInstancedStructDataDetails::OnStructHandlePostChange()
 {
-	if (!StructProperty)
+	if (!bIsHandlingStructValuePostChange)
 	{
-		return nullptr;
+		// External change; force a sync next Tick
+		LastSyncEditableInstanceFromSourceSeconds = 0.0;
 	}
-	TArray<void*> RawData;
-	StructProperty->AccessRawData(RawData);
-	return RawData.Num() == 1 ? static_cast<FInstancedStruct*>(RawData[0]) : nullptr;
 }
 
-void FInstancedStructDataDetails::OnStructPreChange()
+void FInstancedStructDataDetails::SyncEditableInstanceFromSource(bool* OutStructMismatch)
 {
-	if (StructProperty)
+	if (OutStructMismatch)
+	{
+		*OutStructMismatch = false;
+	}
+
+	if (StructProperty && StructProperty->IsValidHandle())
 	{
-		TArray<UObject*> OuterObjects;
-		StructProperty->GetOuterObjects(OuterObjects);
-		for (UObject* Outer : OuterObjects)
+		const UScriptStruct* ExpectedStructType = StructInstanceData ? Cast<UScriptStruct>(StructInstanceData->GetStruct()) : nullptr;
+		StructProperty->EnumerateConstRawData([this, ExpectedStructType, OutStructMismatch](const void* RawData, const int32 /*DataIndex*/, const int32 NumDatas)
 		{
-			if (Outer)
+			if (RawData && NumDatas == 1)
 			{
-				Outer->Modify();
+				const FInstancedStruct* InstancedStruct = static_cast<const FInstancedStruct*>(RawData);
+
+				// Only copy the data if this source is still using the expected struct type
+				const UScriptStruct* StructTypePtr = InstancedStruct->GetScriptStruct();
+				if (StructTypePtr == ExpectedStructType)
+				{
+					if (StructTypePtr)
+					{
+						StructTypePtr->CopyScriptStruct(StructInstanceData->GetStructMemory(), InstancedStruct->GetMemory());
+					}
+				}
+				else if (OutStructMismatch)
+				{
+					*OutStructMismatch = true;
+				}
 			}
-		}
+			return false;
+		});
 	}
+
+	LastSyncEditableInstanceFromSourceSeconds = FPlatformTime::Seconds();
 }
 
-void FInstancedStructDataDetails::OnStructChanged()
+void FInstancedStructDataDetails::SyncEditableInstanceToSource(bool* OutStructMismatch)
 {
-	if (StructProperty)
+	if (OutStructMismatch)
 	{
-		StructProperty->NotifyFinishedChangingProperties();
+		*OutStructMismatch = false;
+	}
+
+	if (StructProperty && StructProperty->IsValidHandle())
+	{
+		const UScriptStruct* ExpectedStructType = StructInstanceData ? Cast<UScriptStruct>(StructInstanceData->GetStruct()) : nullptr;
+		StructProperty->EnumerateRawData([this, ExpectedStructType, OutStructMismatch](const void* RawData, const int32 /*DataIndex*/, const int32 /*NumDatas*/)
+		{
+			if (RawData)
+			{
+				const FInstancedStruct* InstancedStruct = static_cast<const FInstancedStruct*>(RawData);
+
+				// Only copy the data if this source is still using the expected struct type
+				const UScriptStruct* StructTypePtr = InstancedStruct->GetScriptStruct();
+				if (StructTypePtr == ExpectedStructType)
+				{
+					if (StructTypePtr)
+					{
+						StructTypePtr->CopyScriptStruct(InstancedStruct->GetMutableMemory(), StructInstanceData->GetStructMemory());
+					}
+				}
+				else if (OutStructMismatch)
+				{
+					*OutStructMismatch = true;
+				}
+			}
+			return false;
+		});
 	}
 }
 
@@ -118,55 +218,65 @@ void FInstancedStructDataDetails::SetOnRebuildChildren(FSimpleDelegate InOnRegen
 
 void FInstancedStructDataDetails::GenerateHeaderRowContent(FDetailWidgetRow& NodeRow)
 {
+	StructProperty->SetOnPropertyValueChanged(FSimpleDelegate::CreateSP(this, &FInstancedStructDataDetails::OnStructHandlePostChange));
 }
 
 void FInstancedStructDataDetails::GenerateChildContent(IDetailChildrenBuilder& ChildBuilder)
 {
-	FInstancedStruct* InstancedStruct = GetSingleInstancedStruct();
-	if (!InstancedStruct)
+	// Create a struct instance to edit, for the common struct type of the sources being edited
+	StructInstanceData.Reset();
+	if (const UScriptStruct* CommonStructType = UE::StructUtils::Private::GetCommonScriptStruct(StructProperty))
 	{
-		return;
-	}
+		StructInstanceData = MakeShared<FStructOnScope>(CommonStructType);
 
-	const UScriptStruct* ScriptStruct = InstancedStruct->GetScriptStruct();
-	uint8* Memory = InstancedStruct->GetMutableMemory();
-	if (!ScriptStruct || !Memory)
-	{
-		return;
+		// Make sure the struct also has a valid package set, so that properties that rely on this (like FText) work correctly
+		{
+			TArray<UPackage*> OuterPackages;
+			StructProperty->GetOuterPackages(OuterPackages);
+			if (OuterPackages.Num() > 0)
+			{
+				StructInstanceData->SetPackage(OuterPackages[0]);
+			}
+		}
 	}
 
-	// Note: this is potentially dangerous. It puts pointer to external memory on a shared ref.
-	TSharedRef<FStructOnScope> Struct = MakeShareable(new FStructOnScope(ScriptStruct, Memory));
-
-	TArray<UPackage*> OuterPackages;
-	StructProperty->GetOuterPackages(OuterPackages);
-	if (OuterPackages.Num() > 0)
+	// Add the rows for the struct
+	if (StructInstanceData)
 	{
-		Struct->SetPackage(OuterPackages[0]);
-	}
+		FSimpleDelegate OnStructValuePreChangeDelegate = FSimpleDelegate::CreateSP(this, &FInstancedStructDataDetails::OnStructValuePreChange);
+		FSimpleDelegate OnStructValuePostChangeDelegate = FSimpleDelegate::CreateSP(this, &FInstancedStructDataDetails::OnStructValuePostChange);
 
-	static const FName ShowOnlyInnerPropertiesMeta(TEXT("ShowOnlyInnerProperties"));
-	StructProperty->SetInstanceMetaData(ShowOnlyInnerPropertiesMeta, FString());
-
-	TArray<TSharedPtr<IPropertyHandle>> ChildProperties = StructProperty->AddChildStructure(Struct);
-	for (TSharedPtr<IPropertyHandle> ChildHandle : ChildProperties)
-	{
-		ChildHandle->SetOnPropertyValuePreChange(FSimpleDelegate::CreateSP(this, &FInstancedStructDataDetails::OnStructPreChange));
-		ChildHandle->SetOnChildPropertyValuePreChange(FSimpleDelegate::CreateSP(this, &FInstancedStructDataDetails::OnStructPreChange));
-		ChildHandle->SetOnPropertyValueChanged(FSimpleDelegate::CreateSP(this, &FInstancedStructDataDetails::OnStructChanged));
-		ChildHandle->SetOnChildPropertyValueChanged(FSimpleDelegate::CreateSP(this, &FInstancedStructDataDetails::OnStructChanged));
+		// Note: We use AddExternalStructureProperty here as it gives the desired result (the struct value properties as direct children of the struct header)
+		// Neither AddExternalStructure (which added an extra row) nor AddAllExternalStructureProperties (which didn't indent the properties as children) were suitable
+		for (TFieldIterator<const FProperty> PropertyIt(StructInstanceData->GetStruct()); PropertyIt; ++PropertyIt)
+		{
+			if (IDetailPropertyRow* StructValuePropertyRow = ChildBuilder.AddExternalStructureProperty(StructInstanceData.ToSharedRef(), PropertyIt->GetFName()))
+			{
+				TSharedPtr<IPropertyHandle> StructValuePropertyHandle = StructValuePropertyRow->GetPropertyHandle();
+				StructValuePropertyHandle->SetOnPropertyValuePreChange(OnStructValuePreChangeDelegate);
+				StructValuePropertyHandle->SetOnChildPropertyValuePreChange(OnStructValuePreChangeDelegate);
+				StructValuePropertyHandle->SetOnPropertyValueChanged(OnStructValuePostChangeDelegate);
+				StructValuePropertyHandle->SetOnChildPropertyValueChanged(OnStructValuePostChangeDelegate);
 
-		IDetailPropertyRow& Row = ChildBuilder.AddProperty(ChildHandle.ToSharedRef());
-		OnChildRowAdded(Row);
+				OnChildRowAdded(*StructValuePropertyRow);
+			}
+		}
 	}
 }
 
 void FInstancedStructDataDetails::Tick(float DeltaTime)
 {
-	if (bRefresh)
+	if (LastSyncEditableInstanceFromSourceSeconds + 0.1 < FPlatformTime::Seconds())
 	{
-		OnRegenerateChildren.ExecuteIfBound();
-		bRefresh = false;
+		bool bStructMismatch = false;
+		SyncEditableInstanceFromSource(&bStructMismatch);
+
+		if (bStructMismatch)
+		{
+			// If the editable struct no longer has the same struct type as the underlying source, 
+			// then we need to refresh to update the child property rows for the new type
+			OnRegenerateChildren.ExecuteIfBound();
+		}
 	}
 }
 
@@ -178,16 +288,14 @@ FName FInstancedStructDataDetails::GetName() const
 
 void FInstancedStructDataDetails::PostUndo(bool bSuccess)
 {
-	// Regen here directly so that we do not access stale data from FStructOnScope created on GenerateChildContent().
-	OnRegenerateChildren.ExecuteIfBound();
-	bRefresh = false;
+	// Undo; force a sync next Tick
+	LastSyncEditableInstanceFromSourceSeconds = 0.0;
 }
 
 void FInstancedStructDataDetails::PostRedo(bool bSuccess)
 {
-	// Regen here directly so that we do not access stale data from FStructOnScope created on GenerateChildContent().
-	OnRegenerateChildren.ExecuteIfBound();
-	bRefresh = false;
+	// Redo; force a sync next Tick
+	LastSyncEditableInstanceFromSourceSeconds = 0.0;
 }
 
 ////////////////////////////////////
@@ -195,25 +303,32 @@ void FInstancedStructDataDetails::PostRedo(bool bSuccess)
 
 TSharedRef<IPropertyTypeCustomization> FInstancedStructDetails::MakeInstance()
 {
-	return MakeShareable(new FInstancedStructDetails);
+	return MakeShared<FInstancedStructDetails>();
 }
 
 void FInstancedStructDetails::CustomizeHeader(TSharedRef<class IPropertyHandle> StructPropertyHandle, class FDetailWidgetRow& HeaderRow, IPropertyTypeCustomizationUtils& StructCustomizationUtils)
 {
+	static const FName NAME_BaseStruct = "BaseStruct";
+	static const FName NAME_StructTypeConst = "StructTypeConst";
+
 	StructProperty = StructPropertyHandle;
-	PropUtils = StructCustomizationUtils.GetPropertyUtilities().Get();
 
-	// Find base class from meta data.
-	static const FName BaseClassMetaName(TEXT("BaseStruct"));
-	FString BaseClassName = StructProperty->GetMetaData(BaseClassMetaName);
-	BaseScriptStruct = FindObject<UScriptStruct>(ANY_PACKAGE, *BaseClassName);
+	const FProperty* MetaDataProperty = StructProperty->GetMetaDataProperty();
 
-	static const FName StructTypeConst(TEXT("StructTypeConst"));
-	const bool bEnableStructSelection = !StructProperty->HasMetaData(StructTypeConst);
+	const bool bEnableStructSelection = !MetaDataProperty->HasMetaData(NAME_StructTypeConst);
 
-	const FIsResetToDefaultVisible IsResetVisible = FIsResetToDefaultVisible::CreateSP(this, &FInstancedStructDetails::ShouldResetToDefault);
-	const FResetToDefaultHandler ResetHandler = FResetToDefaultHandler::CreateSP(this, &FInstancedStructDetails::ResetToDefault);
-	const FResetToDefaultOverride ResetOverride = FResetToDefaultOverride::Create(IsResetVisible, ResetHandler);
+	BaseScriptStruct = nullptr;
+	{
+		const FString& BaseStructName = MetaDataProperty->GetMetaData(NAME_BaseStruct);
+		if (!BaseStructName.IsEmpty())
+		{
+			BaseScriptStruct = FindObject<UScriptStruct>(ANY_PACKAGE, *BaseStructName);
+			if (!BaseScriptStruct)
+			{
+				BaseScriptStruct = LoadObject<UScriptStruct>(nullptr, *BaseStructName);
+			}
+		}
+	}
 
 	HeaderRow
 		.NameContent()
@@ -244,111 +359,27 @@ void FInstancedStructDetails::CustomizeHeader(TSharedRef<class IPropertyHandle>
 				[
 					SNew(STextBlock)
 					.Text(this, &FInstancedStructDetails::GetDisplayValueString)
+					.ToolTipText(this, &FInstancedStructDetails::GetDisplayValueString)
 					.Font(IDetailLayoutBuilder::GetDetailFont())
 				]
 			]
-		]
-		.OverrideResetToDefault(ResetOverride);
-}
-
-
-bool FInstancedStructDetails::ShouldResetToDefault(TSharedPtr<IPropertyHandle> PropertyHandle) const
-{
-	check(StructProperty);
-	
-	bool bAnyValid = false;
-	
-	TArray<void*> RawData;
-	StructProperty->AccessRawData(RawData);
-	for (void* Data : RawData)
-	{
-		if (FInstancedStruct* Struct = static_cast<FInstancedStruct*>(Data))
-		{
-			if (Struct->IsValid())
-			{
-				bAnyValid = true;
-				break;
-			}
-		}
-	}
-	
-	// Assume that the default value is empty. Any valid means that some can be reset to empty.
-	return bAnyValid;
-}
-
-void FInstancedStructDetails::ResetToDefault(TSharedPtr<IPropertyHandle> PropertyHandle)
-{
-	check(StructProperty);
-	
-	GEditor->BeginTransaction(LOCTEXT("OnResetToDefault", "Reset to default"));
-
-	StructProperty->NotifyPreChange();
-
-	TArray<void*> RawData;
-	StructProperty->AccessRawData(RawData);
-	for (void* Data : RawData)
-	{
-		if (FInstancedStruct* Struct = static_cast<FInstancedStruct*>(Data))
-		{
-			// Assume that the default value is empty.
-			Struct->Reset();
-		}
-	}
-
-	StructProperty->NotifyPostChange(EPropertyChangeType::ValueSet);
-	GEditor->EndTransaction();
-	StructProperty->NotifyFinishedChangingProperties();
-
-	if (PropUtils)
-	{
-		PropUtils->ForceRefresh();
-	}
+		];
 }
 
 void FInstancedStructDetails::CustomizeChildren(TSharedRef<class IPropertyHandle> StructPropertyHandle, class IDetailChildrenBuilder& StructBuilder, IPropertyTypeCustomizationUtils& StructCustomizationUtils)
 {
-	TSharedRef<FInstancedStructDataDetails> DataDetails = MakeShareable(new FInstancedStructDataDetails(StructProperty));
+	TSharedRef<FInstancedStructDataDetails> DataDetails = MakeShared<FInstancedStructDataDetails>(StructProperty);
 	StructBuilder.AddCustomBuilder(DataDetails);
 }
 
-const UScriptStruct* FInstancedStructDetails::GetCommonScriptStruct() const
-{
-	if (!StructProperty)
-	{
-		return nullptr;
-	}
-	const UScriptStruct* CommonScriptStruct = nullptr;
-	bool bMultipleValues = false;
-	TArray<void*> RawData;
-	StructProperty->AccessRawData(RawData);
-	for (void* Data : RawData)
-	{
-		if (FInstancedStruct* Struct = static_cast<FInstancedStruct*>(Data))
-		{
-			const UScriptStruct* ScriptStruct = Struct->GetScriptStruct();
-			if (!bMultipleValues && !CommonScriptStruct)
-			{
-				CommonScriptStruct = ScriptStruct;
-			}
-			else if (ScriptStruct != CommonScriptStruct)
-			{
-				CommonScriptStruct = nullptr;
-				bMultipleValues = true;
-			}
-		}
-	}
-
-	return CommonScriptStruct;
-}
-
 FText FInstancedStructDetails::GetDisplayValueString() const
 {
-	const UScriptStruct* ScriptStruct = GetCommonScriptStruct();
+	const UScriptStruct* ScriptStruct = UE::StructUtils::Private::GetCommonScriptStruct(StructProperty);
 	if (ScriptStruct)
 	{
 		return ScriptStruct->GetDisplayNameText();
 	}
-	return FText();
+	return LOCTEXT("NullScriptStruct", "None");
 }
 
 const FSlateBrush* FInstancedStructDetails::GetDisplayValueIcon() const
@@ -358,8 +389,16 @@ const FSlateBrush* FInstancedStructDetails::GetDisplayValueIcon() const
 
 TSharedRef<SWidget> FInstancedStructDetails::GenerateStructPicker()
 {
-	static const FName ExcludeBaseStruct(TEXT("ExcludeBaseStruct"));
-	const bool bExcludeBaseStruct = StructProperty->HasMetaData(ExcludeBaseStruct);
+	static const FName NAME_ExcludeBaseStruct = "ExcludeBaseStruct";
+	static const FName NAME_HideViewOptions = "HideViewOptions";
+	static const FName NAME_ShowTreeView = "ShowTreeView";
+
+	const FProperty* MetaDataProperty = StructProperty->GetMetaDataProperty();
+
+	const bool bExcludeBaseStruct = MetaDataProperty->HasMetaData(NAME_ExcludeBaseStruct);
+	const bool bAllowNone = !(MetaDataProperty->PropertyFlags & CPF_NoClear);
+	const bool bHideViewOptions = MetaDataProperty->HasMetaData(NAME_HideViewOptions);
+	const bool bShowTreeView = MetaDataProperty->HasMetaData(NAME_ShowTreeView);
 
 	TSharedRef<FInstancedStructFilter> StructFilter = MakeShared<FInstancedStructFilter>();
 	StructFilter->BaseStruct = BaseScriptStruct;
@@ -367,12 +406,11 @@ TSharedRef<SWidget> FInstancedStructDetails::GenerateStructPicker()
 	StructFilter->bAllowBaseStruct = !bExcludeBaseStruct;
 
 	FStructViewerInitializationOptions Options;
-	Options.bShowUnloadedStructs = true;
-	Options.bShowNoneOption = true;
+	Options.bShowNoneOption = bAllowNone;
 	Options.StructFilter = StructFilter;
 	Options.NameTypeToDisplay = EStructViewerNameTypeToDisplay::DisplayName;
-	Options.DisplayMode = EStructViewerDisplayMode::ListView;
-	Options.bAllowViewOptions = true;
+	Options.DisplayMode = bShowTreeView ? EStructViewerDisplayMode::TreeView : EStructViewerDisplayMode::ListView;
+	Options.bAllowViewOptions = !bHideViewOptions;
 
 	FOnStructPicked OnPicked(FOnStructPicked::CreateRaw(this, &FInstancedStructDetails::OnStructPicked));
 
@@ -391,42 +429,26 @@ TSharedRef<SWidget> FInstancedStructDetails::GenerateStructPicker()
 
 void FInstancedStructDetails::OnStructPicked(const UScriptStruct* InStruct)
 {
-	if (StructProperty)
+	if (StructProperty && StructProperty->IsValidHandle())
 	{
-		GEditor->BeginTransaction(LOCTEXT("OnStructPicked", "Set Struct"));
+		FScopedTransaction Transaction(LOCTEXT("OnStructPicked", "Set Struct"));
 
 		StructProperty->NotifyPreChange();
 
-		TArray<void*> RawData;
-		StructProperty->AccessRawData(RawData);
-		for (void* Data : RawData)
+		StructProperty->EnumerateRawData([InStruct](void* RawData, const int32 /*DataIndex*/, const int32 /*NumDatas*/)
 		{
-			if (FInstancedStruct* Struct = static_cast<FInstancedStruct*>(Data))
+			if (FInstancedStruct* InstancedStruct = static_cast<FInstancedStruct*>(RawData))
 			{
-				if (InStruct)
-				{
-					Struct->InitializeAs(InStruct);
-				}
-				else
-				{
-					Struct->Reset();
-				}
+				InstancedStruct->InitializeAs(InStruct);
 			}
-		}
+			return true;
+		});
 
 		StructProperty->NotifyPostChange(EPropertyChangeType::ValueSet);
-		GEditor->EndTransaction();
 		StructProperty->NotifyFinishedChangingProperties();
-
-		GEditor->EndTransaction();
 	}
 
 	ComboButton->SetIsOpen(false);
-
-	if (PropUtils)
-	{
-		PropUtils->ForceRefresh();
-	}
 }
 
 #undef LOCTEXT_NAMESPACE
diff --git a/Engine/Plugins/Experimental/StructUtils/Source/StructUtilsEditor/Public/InstancedStructDetails.h b/Engine/Plugins/Experimental/StructUtils/Source/StructUtilsEditor/Public/InstancedStructDetails.h
--- a/Engine/Plugins/Experimental/StructUtils/Source/StructUtilsEditor/Public/InstancedStructDetails.h
+++ b/Engine/Plugins/Experimental/StructUtils/Source/StructUtilsEditor/Public/InstancedStructDetails.h
@@ -9,8 +9,9 @@
 #include "EditorUndoClient.h"
 
 class IPropertyHandle;
-struct FInstancedStruct;
 class IDetailPropertyRow;
+class FStructOnScope;
+struct FInstancedStruct;
 
 /**
  * Type customization for FInstancedStruct.
@@ -26,17 +27,17 @@ public:
 	virtual void CustomizeChildren(TSharedRef<IPropertyHandle> StructPropertyHandle, class IDetailChildrenBuilder& StructBuilder, IPropertyTypeCustomizationUtils& StructCustomizationUtils) override;
 
 private:
-	bool ShouldResetToDefault(TSharedPtr<IPropertyHandle> PropertyHandle) const;
-	void ResetToDefault(TSharedPtr<IPropertyHandle> PropertyHandle);
-	const UScriptStruct* GetCommonScriptStruct() const;
 	FText GetDisplayValueString() const;
 	const FSlateBrush* GetDisplayValueIcon() const;
 	TSharedRef<SWidget> GenerateStructPicker();
 	void OnStructPicked(const UScriptStruct* InStruct);
 
-	class IPropertyUtilities* PropUtils;
+	/** Handle to the struct property being edited */
 	TSharedPtr<IPropertyHandle> StructProperty;
+	
+	/** The base struct that we're allowing to be picked (controlled by the "BaseStruct" meta-data) */
 	UScriptStruct* BaseScriptStruct = nullptr;
+
 	TSharedPtr<class SComboButton> ComboButton;
 };
 
@@ -69,11 +70,27 @@ public:
 	virtual void OnChildRowAdded(IDetailPropertyRow& ChildRow) {}
 
 private:
-	void OnStructPreChange();
-	void OnStructChanged();
-	FInstancedStruct* GetSingleInstancedStruct();
+	/** Pre/Post change notifications for struct value changes */
+	void OnStructValuePreChange();
+	void OnStructValuePostChange();
+	void OnStructHandlePostChange();
+
+	/** Sync the current state of the editable struct instance to/from the source instance(s) */
+	void SyncEditableInstanceFromSource(bool* OutStructMismatch = nullptr);
+	void SyncEditableInstanceToSource(bool* OutStructMismatch = nullptr);
 
+	/** Handle to the struct property being edited */
 	TSharedPtr<IPropertyHandle> StructProperty;
+
+	/** Struct instance that is being edited; this is a copy of the source struct data to avoid lifetime issues when the underlying source is updated/deleted */
+	TSharedPtr<FStructOnScope> StructInstanceData;
+
+	/** Delegate that can be used to refresh the child rows of the current struct (eg, when changing struct type) */
 	FSimpleDelegate OnRegenerateChildren;
-	bool bRefresh = false;
+
+	/** The last time that SyncEditableInstanceFromSource was called, in FPlatformTime::Seconds() */
+	double LastSyncEditableInstanceFromSourceSeconds = 0.0;
+
+	/** True if we're currently handling a StructValuePostChange */
+	bool bIsHandlingStructValuePostChange = false;
 };
