Deleted FStructVariant in favor of using FInstancedStruct

diff --git a/Engine/Plugins/Experimental/StructUtils/Source/StructUtils/Private/InstancedStruct.cpp b/Engine/Plugins/Experimental/StructUtils/Source/StructUtils/Private/InstancedStruct.cpp
--- a/Engine/Plugins/Experimental/StructUtils/Source/StructUtils/Private/InstancedStruct.cpp
+++ b/Engine/Plugins/Experimental/StructUtils/Source/StructUtils/Private/InstancedStruct.cpp
@@ -2,7 +2,6 @@
 #include "InstancedStruct.h"
 #include "StructView.h"
 #include "SharedStruct.h"
-#include "UObject/StructVariant.h"
 #include "Serialization/PropertyLocalizationDataGathering.h"
 
 namespace UE::StructUtils::Private
@@ -260,11 +259,62 @@ bool FInstancedStruct::SerializeFromMismatchedTag(const FPropertyTag& Tag, FStru
 	static const FName NAME_StructVariant = "StructVariant";
 	if (Tag.Type == NAME_StructProperty && Tag.StructName == NAME_StructVariant)
 	{
-		// TODO: When deleting FStructVariant, migrate its Serialize loading code here
-		FStructVariant StructVariant;
-		StructVariant.Serialize(Slot);
+		auto SerializeStructVariant = [this, &Tag, &Slot]()
+		{
+			FArchive& UnderlyingArchive = Slot.GetUnderlyingArchive();
+			FStructuredArchive::FRecord Record = Slot.EnterRecord();
+
+			// Serialize the struct type
+			UScriptStruct* StructTypePtr = nullptr;
+			Record << SA_VALUE(TEXT("StructType"), StructTypePtr);
+			if (StructTypePtr)
+			{
+				UnderlyingArchive.Preload(StructTypePtr);
+			}
+			InitializeAs(StructTypePtr);
+
+			auto SerializeStructInstance = [this, StructTypePtr, &Record]()
+			{
+				if (StructTypePtr)
+				{
+					StructTypePtr->SerializeItem(Record.EnterField(SA_FIELD_NAME(TEXT("StructInstance"))), GetMutableMemory(), nullptr);
+				}
+			};
+
+			// Serialize the struct instance, potentially tagging it with its serialized size 
+			// in-case the struct is deleted later and we need to step over the instance data
+			const bool bTagStructInstance = !UnderlyingArchive.IsTextFormat();
+			if (bTagStructInstance)
+			{
+				// Read the serialized size
+				int64 StructInstanceSerializedSize = 0;
+				UnderlyingArchive << StructInstanceSerializedSize;
+
+				// Serialize the struct instance
+				const int64 StructInstanceStartOffset = UnderlyingArchive.Tell();
+				SerializeStructInstance();
+				const int64 StructInstanceEndOffset = UnderlyingArchive.Tell();
+
+				// Ensure we're at the correct location after serializing the instance data
+				const int64 ExpectedStructInstanceEndOffset = StructInstanceStartOffset + StructInstanceSerializedSize;
+				if (StructInstanceEndOffset != ExpectedStructInstanceEndOffset)
+				{
+					if (StructTypePtr)
+					{
+						// We only expect a mismatch here if the underlying struct is no longer available!
+						UnderlyingArchive.SetCriticalError();
+						UE_LOG(LogCore, Error, TEXT("FStructVariant expected to read %lld bytes for struct %s but read %lld bytes!"), StructInstanceSerializedSize, *StructTypePtr->GetName(), StructInstanceEndOffset - StructInstanceStartOffset);
+					}
+					UnderlyingArchive.Seek(ExpectedStructInstanceEndOffset);
+				}
+			}
+			else
+			{
+				SerializeStructInstance();
+			}
+		};
 
-		InitializeAs(StructVariant.GetStructType(), static_cast<uint8*>(StructVariant.GetStructInstance()));
+		SerializeStructVariant();
 		return true;
 	}
 
diff --git a/Engine/Source/Editor/DetailCustomizations/Private/DetailCustomizations.cpp b/Engine/Source/Editor/DetailCustomizations/Private/DetailCustomizations.cpp
--- a/Engine/Source/Editor/DetailCustomizations/Private/DetailCustomizations.cpp
+++ b/Engine/Source/Editor/DetailCustomizations/Private/DetailCustomizations.cpp
@@ -137,7 +137,6 @@
 #include "SplineComponentDetails.h"
 #include "StaticMeshActorDetails.h"
 #include "StaticMeshComponentDetails.h"
-#include "StructVariantCustomization.h"
 #include "SubmixDetailsCustomization.h"
 #include "SupportedRangeTypes.h"	// StructsSupportingRangeVisibility
 #include "Templates/SharedPointer.h"
@@ -262,7 +261,6 @@ void FDetailCustomizationsModule::RegisterPropertyTypeCustomizations()
 	RegisterCustomPropertyTypeLayout("TextProperty", FOnGetPropertyTypeCustomizationInstance::CreateStatic(&FTextCustomization::MakeInstance));
 	RegisterCustomPropertyTypeLayout("DirectoryPath", FOnGetPropertyTypeCustomizationInstance::CreateStatic(&FDirectoryPathStructCustomization::MakeInstance));
 	RegisterCustomPropertyTypeLayout("FilePath", FOnGetPropertyTypeCustomizationInstance::CreateStatic(&FFilePathStructCustomization::MakeInstance));
-	RegisterCustomPropertyTypeLayout("StructVariant", FOnGetPropertyTypeCustomizationInstance::CreateStatic(&FStructVariantCustomization::MakeInstance));
 	RegisterCustomPropertyTypeLayout("IOSBuildResourceDirectory", FOnGetPropertyTypeCustomizationInstance::CreateStatic(&FDirectoryPathStructCustomization::MakeInstance));
 	RegisterCustomPropertyTypeLayout("IOSBuildResourceFilePath", FOnGetPropertyTypeCustomizationInstance::CreateStatic(&FFilePathStructCustomization::MakeInstance));
 	RegisterCustomPropertyTypeLayout("InputAxisConfigEntry", FOnGetPropertyTypeCustomizationInstance::CreateStatic(&FInputAxisConfigCustomization::MakeInstance));
diff --git a/Engine/Source/Editor/DetailCustomizations/Private/StructVariantCustomization.cpp b/Engine/Source/Editor/DetailCustomizations/Private/StructVariantCustomization.cpp
deleted file mode 100644
diff --git a/Engine/Source/Editor/DetailCustomizations/Private/StructVariantCustomization.h b/Engine/Source/Editor/DetailCustomizations/Private/StructVariantCustomization.h
deleted file mode 100644
diff --git a/Engine/Source/Runtime/CoreUObject/Private/UObject/StructVariant.cpp b/Engine/Source/Runtime/CoreUObject/Private/UObject/StructVariant.cpp
deleted file mode 100644
diff --git a/Engine/Source/Runtime/CoreUObject/Public/UObject/StructVariant.h b/Engine/Source/Runtime/CoreUObject/Public/UObject/StructVariant.h
deleted file mode 100644
