FStructArrayView created and initial unit test added.

diff --git a/Engine/Plugins/Experimental/StructUtils/Source/StructUtils/Public/StructArrayView.h b/Engine/Plugins/Experimental/StructUtils/Source/StructUtils/Public/StructArrayView.h
new file mode 100644
--- /dev/null
+++ b/Engine/Plugins/Experimental/StructUtils/Source/StructUtils/Public/StructArrayView.h
@@ -0,0 +1,105 @@
+// Copyright Epic Games, Inc. All Rights Reserved.
+
+#pragma once
+
+#include "StructUtils.h"
+#include "InstancedStruct.h"
+#include "SharedStruct.h"
+
+///////////////////////////////////////////////////////////////// FStructArrayView /////////////////////////////////////////////////////////////////
+
+/** 
+ *	A generic, transient view of a homogeneously-typed array of instances of a specific UScriptStruct
+ */
+struct FStructArrayView
+{
+	FStructArrayView(FStructArrayView&& Src)
+		: ElementSize(Src.ElementSize)
+		, DataPtr(Src.DataPtr)
+		, NumElements(Src.NumElements)
+		, FragmentType(Src.FragmentType)
+	{}
+
+	FStructArrayView(const FStructArrayView& Src)
+		: ElementSize(Src.ElementSize)
+		, DataPtr(Src.DataPtr)
+		, NumElements(Src.NumElements)
+		, FragmentType(Src.FragmentType)
+	{}
+	
+	template<typename T>
+	explicit FStructArrayView(TArray<T>& InArray)
+		: ElementSize(sizeof(T))
+		, DataPtr(InArray.GetData())
+		, NumElements(InArray.Num())
+		, FragmentType(*StaticStruct<typename TRemoveReference<T>::Type>())
+
+	{}
+
+	template<typename T>
+	explicit FStructArrayView(TArrayView<T> InArrayView)
+		: ElementSize(sizeof(T))
+		, DataPtr(InArrayView.GetData())
+		, NumElements(InArrayView.Num())
+		, FragmentType(*StaticStruct<typename TRemoveReference<T>::Type>())
+	{}
+
+	/**
+	 * Creates a sub view of the given view
+	 */
+	FStructArrayView(const FStructArrayView& Src, const int32 OffsetElements, const int32 InNumElements)
+		: ElementSize(Src.ElementSize)
+		, DataPtr((uint8*)Src.DataPtr + Src.ElementSize * OffsetElements)
+		, NumElements(InNumElements)
+		, FragmentType(Src.FragmentType)
+	{
+		checkf(InNumElements + OffsetElements <= Src.NumElements, TEXT("Requested range passes over the end of the view, %d + %d> %d")
+			, InNumElements, OffsetElements, Src.NumElements);
+	}
+
+	SIZE_T GetElementSize() const { return ElementSize; }
+	const void* GetData() const { return DataPtr; }
+	const void* GetDataAt(const int32 Index) const { return (uint8*)DataPtr + Index * ElementSize; }
+	void* GetMutableDataAt(const int32 Index) const { return (uint8*)DataPtr + Index * ElementSize; }
+	int32 Num() const { return NumElements; }
+	const UScriptStruct& GetFragmentType() const { return FragmentType; }
+
+	template<typename T>
+	const T& GetElementAt(const int32 Index) const
+	{
+		return *((T*)GetDataAt(Index));
+	}
+
+	template<typename T>
+	const T& GetElementAtChecked(const int32 Index) const
+	{
+		check(T::StaticStruct() == &FragmentType);
+		return *((T*)GetDataAt(Index));
+	}
+
+	template<typename T>
+	T& GetMutableElementAt(const int32 Index)
+	{
+		return *((T*)GetDataAt(Index));
+	}
+
+	template<typename T>
+	T& GetMutableElementAtChecked(const int32 Index)
+	{
+		check(T::StaticStruct() == &FragmentType);
+		return *((T*)GetDataAt(Index));
+	}
+
+	FORCEINLINE void Swap(const int32 A, const int32 B)
+	{
+		checkSlow(A >= 0 && A < NumElements);
+		checkSlow(B >= 0 && B < NumElements);
+		FMemory::Memswap(GetMutableDataAt(A), GetMutableDataAt(B), ElementSize);
+	}
+
+private:
+	SIZE_T ElementSize = 0;
+	void* DataPtr = nullptr;
+	int32 NumElements = 0;
+	UScriptStruct& FragmentType;
+};
diff --git a/Engine/Plugins/Experimental/StructUtils/Source/StructUtilsTestSuite/Private/StructArrayViewTest.cpp b/Engine/Plugins/Experimental/StructUtils/Source/StructUtilsTestSuite/Private/StructArrayViewTest.cpp
new file mode 100644
--- /dev/null
+++ b/Engine/Plugins/Experimental/StructUtils/Source/StructUtilsTestSuite/Private/StructArrayViewTest.cpp
@@ -0,0 +1,48 @@
+// Copyright Epic Games, Inc. All Rights Reserved.
+
+#include "AITestsCommon.h"
+#include "StructArrayView.h"
+#include "StructUtilsTestTypes.h"
+
+#define LOCTEXT_NAMESPACE "StructUtilsTests"
+
+PRAGMA_DISABLE_OPTIMIZATION
+
+namespace FStructArrayViewTest
+{
+
+struct FTest_ChangeElements : FAITestBase
+{
+	virtual bool InstantTest() override
+	{
+		constexpr float Increment = 10.f;
+		TArray<float> OriginalValues = { 1,2,3,4,5,6,7,8,9,10 };
+		TArray<FTestStructSimple> ValuesF;
+		for (float Value : OriginalValues)
+		{
+			ValuesF.Add(Value);
+		}
+		FStructArrayView ViewF = FStructArrayView(ValuesF);
+
+		for (int i = 0; i < ValuesF.Num(); ++i)
+		{
+			ViewF.GetMutableElementAt<FTestStructSimple>(i).Float += 10.f;
+		}
+
+		for (int i = 0; i < ValuesF.Num(); ++i)
+		{
+			AITEST_EQUAL(TEXT("Element should have the same value independently of access method"), ValuesF[i].Float, ViewF.GetElementAt<FTestStructSimple>(i).Float);
+			AITEST_EQUAL(TEXT("Element should reflect the change done to it"), ViewF.GetElementAt<FTestStructSimple>(i).Float, OriginalValues[i] + Increment);
+			AITEST_EQUAL(TEXT("Element should point at the same adress in memory, independently of access method"), (void*)&ValuesF[i], ViewF.GetDataAt(i));
+		}
+
+		return true;
+	}
+};
+IMPLEMENT_AI_INSTANT_TEST(FTest_ChangeElements, "System.StructUtils.StructArrayView.ChangeElements");
+
+} // FStructArrayViewTest
+
+PRAGMA_ENABLE_OPTIMIZATION
+
+#undef LOCTEXT_NAMESPACE
