Allow AddChildStructure to be used by sparse data properties

The child struct properties can't inherit the sparse data flag, as FStructurePropertyNode is already directly editing a block of struct memory.

diff --git a/Engine/Plugins/Experimental/StructUtils/Source/StructUtilsEditor/Private/Customizations/InstancedStructDetails.cpp b/Engine/Plugins/Experimental/StructUtils/Source/StructUtilsEditor/Private/Customizations/InstancedStructDetails.cpp
--- a/Engine/Plugins/Experimental/StructUtils/Source/StructUtilsEditor/Private/Customizations/InstancedStructDetails.cpp
+++ b/Engine/Plugins/Experimental/StructUtils/Source/StructUtilsEditor/Private/Customizations/InstancedStructDetails.cpp
@@ -246,20 +246,16 @@ void FInstancedStructDataDetails::GenerateChildContent(IDetailChildrenBuilder& C
 		FSimpleDelegate OnStructValuePreChangeDelegate = FSimpleDelegate::CreateSP(this, &FInstancedStructDataDetails::OnStructValuePreChange);
 		FSimpleDelegate OnStructValuePostChangeDelegate = FSimpleDelegate::CreateSP(this, &FInstancedStructDataDetails::OnStructValuePostChange);
 
-		// Note: We use AddExternalStructureProperty here as it gives the desired result (the struct value properties as direct children of the struct header)
-		// Neither AddExternalStructure (which added an extra row) nor AddAllExternalStructureProperties (which didn't indent the properties as children) were suitable
-		for (TFieldIterator<const FProperty> PropertyIt(StructInstanceData->GetStruct()); PropertyIt; ++PropertyIt)
+		TArray<TSharedPtr<IPropertyHandle>> ChildProperties = StructProperty->AddChildStructure(StructInstanceData.ToSharedRef());
+		for (TSharedPtr<IPropertyHandle> ChildHandle : ChildProperties)
 		{
-			if (IDetailPropertyRow* StructValuePropertyRow = ChildBuilder.AddExternalStructureProperty(StructInstanceData.ToSharedRef(), PropertyIt->GetFName()))
-			{
-				TSharedPtr<IPropertyHandle> StructValuePropertyHandle = StructValuePropertyRow->GetPropertyHandle();
-				StructValuePropertyHandle->SetOnPropertyValuePreChange(OnStructValuePreChangeDelegate);
-				StructValuePropertyHandle->SetOnChildPropertyValuePreChange(OnStructValuePreChangeDelegate);
-				StructValuePropertyHandle->SetOnPropertyValueChanged(OnStructValuePostChangeDelegate);
-				StructValuePropertyHandle->SetOnChildPropertyValueChanged(OnStructValuePostChangeDelegate);
+			ChildHandle->SetOnPropertyValuePreChange(OnStructValuePreChangeDelegate);
+			ChildHandle->SetOnChildPropertyValuePreChange(OnStructValuePreChangeDelegate);
+			ChildHandle->SetOnPropertyValueChanged(OnStructValuePostChangeDelegate);
+			ChildHandle->SetOnChildPropertyValueChanged(OnStructValuePostChangeDelegate);
 
-				OnChildRowAdded(*StructValuePropertyRow);
-			}
+			IDetailPropertyRow& Row = ChildBuilder.AddProperty(ChildHandle.ToSharedRef());
+			OnChildRowAdded(Row);
 		}
 	}
 }
diff --git a/Engine/Source/Editor/PropertyEditor/Private/CategoryPropertyNode.cpp b/Engine/Source/Editor/PropertyEditor/Private/CategoryPropertyNode.cpp
--- a/Engine/Source/Editor/PropertyEditor/Private/CategoryPropertyNode.cpp
+++ b/Engine/Source/Editor/PropertyEditor/Private/CategoryPropertyNode.cpp
@@ -101,7 +101,7 @@ void FCategoryPropertyNode::InitChildNodes()
 		InitParams.bAllowChildren = true;
 		InitParams.bForceHiddenPropertyVisibility = bShowHiddenProperties;
 		InitParams.bCreateDisableEditOnInstanceNodes = bShouldShowDisableEditOnInstance;
-		InitParams.bIsSparseProperty = SparseProperties.Contains(Properties[PropertyIndex]);
+		InitParams.IsSparseProperty = SparseProperties.Contains(Properties[PropertyIndex]) ? FPropertyNodeInitParams::EIsSparseDataProperty::True : FPropertyNodeInitParams::EIsSparseDataProperty::Inherit;
 
 		NewItemNode->InitNode( InitParams );
 
diff --git a/Engine/Source/Editor/PropertyEditor/Private/PropertyHandleImpl.cpp b/Engine/Source/Editor/PropertyEditor/Private/PropertyHandleImpl.cpp
--- a/Engine/Source/Editor/PropertyEditor/Private/PropertyHandleImpl.cpp
+++ b/Engine/Source/Editor/PropertyEditor/Private/PropertyHandleImpl.cpp
@@ -3126,6 +3126,7 @@ TArray<TSharedPtr<IPropertyHandle>> FPropertyHandleBase::AddChildStructure( TSha
 	RootInitParams.bAllowChildren = true;
 	RootInitParams.bForceHiddenPropertyVisibility = FPropertySettings::Get().ShowHiddenProperties();
 	RootInitParams.bCreateCategoryNodes = false;
+	RootInitParams.IsSparseProperty = FPropertyNodeInitParams::EIsSparseDataProperty::False; // FStructurePropertyNode can't inherit the sparse data flag
 
 	StructPropertyNode->InitNode(RootInitParams);
 
diff --git a/Engine/Source/Editor/PropertyEditor/Private/PropertyNode.cpp b/Engine/Source/Editor/PropertyEditor/Private/PropertyNode.cpp
--- a/Engine/Source/Editor/PropertyEditor/Private/PropertyNode.cpp
+++ b/Engine/Source/Editor/PropertyEditor/Private/PropertyNode.cpp
@@ -100,9 +100,9 @@ void FPropertyNode::InitNode(const FPropertyNodeInitParams& InitParams)
 	ArrayOffset = InitParams.ArrayOffset;
 	ArrayIndex = InitParams.ArrayIndex;
 
-	bool bIsSparse = InitParams.bIsSparseProperty;
+	bool bIsSparse = InitParams.IsSparseProperty == FPropertyNodeInitParams::EIsSparseDataProperty::True;
 
-	if (ParentNode)
+	if (ParentNode && InitParams.IsSparseProperty == FPropertyNodeInitParams::EIsSparseDataProperty::Inherit)
 	{
 		//default to parents max child depth
 		MaxChildDepthAllowed = ParentNode->MaxChildDepthAllowed;
diff --git a/Engine/Source/Editor/PropertyEditor/Private/PropertyNode.h b/Engine/Source/Editor/PropertyEditor/Private/PropertyNode.h
--- a/Engine/Source/Editor/PropertyEditor/Private/PropertyNode.h
+++ b/Engine/Source/Editor/PropertyEditor/Private/PropertyNode.h
@@ -198,6 +198,13 @@ private:
  */
 struct FPropertyNodeInitParams
 {
+	enum class EIsSparseDataProperty : uint8
+	{
+		False,
+		True,
+		Inherit,
+	};
+
 	/** The parent of the property node */
 	TSharedPtr<FPropertyNode> ParentNode;
 	/** The property that the node observes and modifies*/
@@ -215,7 +222,7 @@ struct FPropertyNodeInitParams
 	/** Whether or not to create nodes for properties marked CPF_DisableEditOnInstance */
 	bool bCreateDisableEditOnInstanceNodes;
 	/** Whether or not this property is sparse data */
-	bool bIsSparseProperty;
+	EIsSparseDataProperty IsSparseProperty;
 
 	FPropertyNodeInitParams()
 		: ParentNode(nullptr)
@@ -226,7 +233,7 @@ struct FPropertyNodeInitParams
 		, bForceHiddenPropertyVisibility( false )
 		, bCreateCategoryNodes( true )
 		, bCreateDisableEditOnInstanceNodes( true )
-		, bIsSparseProperty( false )
+		, IsSparseProperty( EIsSparseDataProperty::Inherit )
 	{}
 };
 
