Fix de-serialization crash in FInstancedStruct for cooked assets

diff --git a/Engine/Plugins/Experimental/StructUtils/Source/StructUtils/Private/InstancedStruct.cpp b/Engine/Plugins/Experimental/StructUtils/Source/StructUtils/Private/InstancedStruct.cpp
--- a/Engine/Plugins/Experimental/StructUtils/Source/StructUtils/Private/InstancedStruct.cpp
+++ b/Engine/Plugins/Experimental/StructUtils/Source/StructUtils/Private/InstancedStruct.cpp
@@ -106,7 +106,7 @@ bool FInstancedStruct::Serialize(FArchive& Ar)
 	bool bUseVersioning = true;
 
 #if WITH_EDITOR
-	if (!Ar.IsCooking())
+	if (!Ar.IsCooking() && !Ar.IsFilterEditorOnly())
 	{
 		// Keep archive position to use legacy serialization if the header is not found
 		const int64 HeaderOffset = Ar.Tell();
