Use TBaseStructure rather than StaticStruct for FConst/StructView and FInstancedStruct so they work with NoExport types

diff --git a/Engine/Plugins/Experimental/StructUtils/Source/StructUtils/Public/InstancedStruct.h b/Engine/Plugins/Experimental/StructUtils/Source/StructUtils/Public/InstancedStruct.h
--- a/Engine/Plugins/Experimental/StructUtils/Source/StructUtils/Public/InstancedStruct.h
+++ b/Engine/Plugins/Experimental/StructUtils/Source/StructUtils/Public/InstancedStruct.h
@@ -81,7 +81,7 @@ public:
 
 		Reset();
 
-		const UScriptStruct* Struct = T::StaticStruct();
+		const UScriptStruct* Struct = TBaseStructure<T>::Get();
 		const int32 RequiredSize = Struct->GetStructureSize();
 		uint8* Memory = (uint8*)FMemory::Malloc(FMath::Max(1, RequiredSize));
 		SetStructData(Struct, Memory);
@@ -96,7 +96,7 @@ public:
 		UE::StructUtils::CheckStructType<T>();
 
 		FInstancedStruct InstancedStruct;
-		InstancedStruct.InitializeAs(T::StaticStruct(), nullptr);
+		InstancedStruct.InitializeAs(TBaseStructure<T>::Get(), nullptr);
 		return InstancedStruct;
 	}
 
@@ -107,7 +107,7 @@ public:
 		UE::StructUtils::CheckStructType<T>();
 
 		FInstancedStruct InstancedStruct;
-		InstancedStruct.InitializeAs(T::StaticStruct(), reinterpret_cast<const uint8*>(&Struct));
+		InstancedStruct.InitializeAs(TBaseStructure<T>::Get(), reinterpret_cast<const uint8*>(&Struct));
 		return InstancedStruct;
 	}
 
@@ -154,7 +154,7 @@ public:
 		const UScriptStruct* Struct = GetScriptStruct();
 		check(Memory != nullptr);
 		check(Struct != nullptr);
-		check(Struct->IsChildOf(T::StaticStruct()));
+		check(Struct->IsChildOf(TBaseStructure<T>::Get()));
 		return *((T*)Memory);
 	}
 
@@ -164,7 +164,7 @@ public:
 	{
 		const uint8* Memory = GetMemory();
 		const UScriptStruct* Struct = GetScriptStruct();
-		if (Memory != nullptr && Struct && Struct->IsChildOf(T::StaticStruct()))
+		if (Memory != nullptr && Struct && Struct->IsChildOf(TBaseStructure<T>::Get()))
 		{
 			return ((T*)Memory);
 		}
@@ -186,7 +186,7 @@ public:
 		const UScriptStruct* Struct = GetScriptStruct();
 		check(Memory != nullptr);
 		check(Struct != nullptr);
-		check(Struct->IsChildOf(T::StaticStruct()));
+		check(Struct->IsChildOf(TBaseStructure<T>::Get()));
 		return *((T*)Memory);
 	}
 
@@ -196,7 +196,7 @@ public:
 	{
 		uint8* Memory = GetMutableMemory();
 		const UScriptStruct* Struct = GetScriptStruct();
-		if (Memory != nullptr && Struct && Struct->IsChildOf(T::StaticStruct()))
+		if (Memory != nullptr && Struct && Struct->IsChildOf(TBaseStructure<T>::Get()))
 		{
 			return ((T*)Memory);
 		}
diff --git a/Engine/Plugins/Experimental/StructUtils/Source/StructUtils/Public/StructView.h b/Engine/Plugins/Experimental/StructUtils/Source/StructUtils/Public/StructView.h
--- a/Engine/Plugins/Experimental/StructUtils/Source/StructUtils/Public/StructView.h
+++ b/Engine/Plugins/Experimental/StructUtils/Source/StructUtils/Public/StructView.h
@@ -50,7 +50,7 @@ public:
 	static FConstStructView Make(const T& Struct)
 	{
 		UE::StructUtils::CheckStructType<T>();
-		return FConstStructView(T::StaticStruct(), reinterpret_cast<const uint8*>(&Struct));
+		return FConstStructView(TBaseStructure<T>::Get(), reinterpret_cast<const uint8*>(&Struct));
 	}
 
 	/** Returns struct type. */
@@ -79,7 +79,7 @@ public:
 		const UScriptStruct* Struct = GetScriptStruct();
 		check(Memory != nullptr);
 		check(Struct != nullptr);
-		check(Struct->IsChildOf(T::StaticStruct()));
+		check(Struct->IsChildOf(TBaseStructure<T>::Get()));
 		return *((T*)Memory);
 	}
 
@@ -89,7 +89,7 @@ public:
 	{
 		const uint8* Memory = GetMemory();
 		const UScriptStruct* Struct = GetScriptStruct();
-		if (Memory != nullptr && Struct && Struct->IsChildOf(T::StaticStruct()))
+		if (Memory != nullptr && Struct && Struct->IsChildOf(TBaseStructure<T>::Get()))
 		{
 			return ((T*)Memory);
 		}
@@ -181,7 +181,7 @@ public:
 	static FStructView Make(T& InStruct)
 	{
 		UE::StructUtils::CheckStructType<T>();
-		return FStructView(T::StaticStruct(), reinterpret_cast<uint8*>(&InStruct));
+		return FStructView(TBaseStructure<T>::Get(), reinterpret_cast<uint8*>(&InStruct));
 	}
 
 	/** Returns a mutable pointer to struct memory. This const_cast here is safe as a ClassName can only be setup from mutable non const memory. */
@@ -199,7 +199,7 @@ public:
 		const UScriptStruct* Struct = GetScriptStruct();
 		check(Memory != nullptr);
 		check(Struct != nullptr);
-		check(Struct->IsChildOf(T::StaticStruct()));
+		check(Struct->IsChildOf(TBaseStructure<T>::Get()));
 		return *((T*)Memory);
 	}
 
@@ -209,7 +209,7 @@ public:
 	{
 		uint8* Memory = GetMutableMemory();
 		const UScriptStruct* Struct = GetScriptStruct();
-		if (Memory != nullptr && Struct && Struct->IsChildOf(T::StaticStruct()))
+		if (Memory != nullptr && Struct && Struct->IsChildOf(TBaseStructure<T>::Get()))
 		{
 			return ((T*)Memory);
 		}
