Remove the API export from FConst/StructView so that it can be used without linking to StructUtils

#rnx

diff --git a/Engine/Plugins/Experimental/StructUtils/Source/StructUtils/Public/StructView.h b/Engine/Plugins/Experimental/StructUtils/Source/StructUtils/Public/StructView.h
--- a/Engine/Plugins/Experimental/StructUtils/Source/StructUtils/Public/StructView.h
+++ b/Engine/Plugins/Experimental/StructUtils/Source/StructUtils/Public/StructView.h
@@ -16,7 +16,7 @@
  * FConstStructView is passed by value.
  * FConstStructView is similar to FStructOnScope, but FConstStructView is a view only (FStructOnScope can either own the memory or be a view)
  */
-struct STRUCTUTILS_API FConstStructView
+struct FConstStructView
 {
 public:
 
@@ -147,7 +147,7 @@ protected:
  * FStructView is passed by value.
  * FStructView is similar to FStructOnScope, but FStructView is a view only (FStructOnScope can either own the memory or be a view)
  */
-struct STRUCTUTILS_API FStructView : public FConstStructView
+struct FStructView : public FConstStructView
 {
 public:
 
