Avoid allocation churn when assigning to an existing FInstancedStruct

FInstancedStruct had both InitializeAs (public) and ReinitializeAs (protected):
 - InitializeAs would always reallocate the struct, and is used by both constructors and assignment operators.
 - ReinitializeAs would avoid reallocating the struct if the type matched, but would fail to return the struct to its default state which could be dangerous if used with delta serialization/import.

This means that if you wanted to avoid allocating churn when reusing an existing FInstancedStruct, you'd have to check the struct type yourself and manually call CopyScriptStruct.

This change removes ReinitializeAs entirely, and has InitializeAs reuse the struct allocation if the new type matches. It will also ensure the struct is in its default state when not given specific data to initialize as.

diff --git a/Engine/Plugins/Experimental/StructUtils/Source/StructUtils/Private/InstancedStruct.cpp b/Engine/Plugins/Experimental/StructUtils/Source/StructUtils/Private/InstancedStruct.cpp
--- a/Engine/Plugins/Experimental/StructUtils/Source/StructUtils/Private/InstancedStruct.cpp
+++ b/Engine/Plugins/Experimental/StructUtils/Source/StructUtils/Private/InstancedStruct.cpp
@@ -57,23 +57,41 @@ FInstancedStruct& FInstancedStruct::operator=(const FConstStructView InOther)
 
 void FInstancedStruct::InitializeAs(const UScriptStruct* InScriptStruct, const uint8* InStructMemory /*= nullptr*/)
 {
-	Reset();
-
-	if (!InScriptStruct)
+	const UScriptStruct* CurrentScriptStruct = GetScriptStruct();
+	if (InScriptStruct && InScriptStruct == CurrentScriptStruct)
 	{
-		// InScriptStruct == nullptr signifies an empty, unset FInstancedStruct instance. No further work required.
-		return;
+		// Struct type already matches...
+		if (InStructMemory)
+		{
+			// ... apply the given state
+			CurrentScriptStruct->CopyScriptStruct(GetMutableMemory(), InStructMemory);
+		}
+		else
+		{
+			// ... return the struct to its default state
+			CurrentScriptStruct->DestroyStruct(GetMutableMemory());
+			CurrentScriptStruct->InitializeStruct(GetMutableMemory());
+		}
 	}
+	else
+	{
+		// Struct type mismatch; reset and reinitialize
+		Reset();
 
-	const int32 RequiredSize = InScriptStruct->GetStructureSize();
-	const uint8* Memory = ((uint8*)FMemory::Malloc(FMath::Max(1, RequiredSize)));
-	SetStructData(InScriptStruct,Memory);
+		// InScriptStruct == nullptr signifies an empty, unset FInstancedStruct instance
+		if (InScriptStruct)
+		{
+			const int32 RequiredSize = InScriptStruct->GetStructureSize();
+			const uint8* Memory = ((uint8*)FMemory::Malloc(FMath::Max(1, RequiredSize)));
+			SetStructData(InScriptStruct, Memory);
 
-	InScriptStruct->InitializeStruct(GetMutableMemory());
+			InScriptStruct->InitializeStruct(GetMutableMemory());
 
-	if (InStructMemory)
-	{
-		InScriptStruct->CopyScriptStruct(GetMutableMemory(), InStructMemory);
+			if (InStructMemory)
+			{
+				InScriptStruct->CopyScriptStruct(GetMutableMemory(), InStructMemory);
+			}
+		}
 	}
 }
 
@@ -144,13 +162,12 @@ bool FInstancedStruct::Serialize(FArchive& Ar)
 	if (Ar.IsLoading())
 	{
 		// UScriptStruct type
-		UScriptStruct* NewNonConstStruct = nullptr;
-		Ar << NewNonConstStruct;
-		if (NewNonConstStruct)
+		Ar << NonConstStruct;
+		if (NonConstStruct)
 		{
-			Ar.Preload(NewNonConstStruct);
+			Ar.Preload(NonConstStruct);
 		}
-		NonConstStruct = ReinitializeAs(NewNonConstStruct);
+		InitializeAs(NonConstStruct);
 
 		// Size of the serialized memory
 		int32 SerialSize = 0; 
@@ -230,7 +247,7 @@ bool FInstancedStruct::ImportTextItem(const TCHAR*& Buffer, int32 PortFlags, UOb
 
 	if (StructPathName.Len() == 0 || FCString::Stricmp(StructPathName.ToString(), TEXT("None")) == 0)
 	{
-		ReinitializeAs(nullptr);
+		InitializeAs(nullptr);
 	}
 	else
 	{
@@ -240,7 +257,7 @@ bool FInstancedStruct::ImportTextItem(const TCHAR*& Buffer, int32 PortFlags, UOb
 			return false;
 		}
 
-		ReinitializeAs(StructTypePtr);
+		InitializeAs(StructTypePtr);
 		if (const TCHAR* Result = StructTypePtr->ImportText(Buffer, GetMutableMemory(), Parent, PortFlags, ErrorText, [StructTypePtr]() { return StructTypePtr->GetName(); }))
 		{
 			Buffer = Result;
@@ -329,19 +346,6 @@ void FInstancedStruct::GetPreloadDependencies(TArray<UObject*>& OutDeps)
 	}
 }
 
-UScriptStruct* FInstancedStruct::ReinitializeAs(const UScriptStruct* InScriptStruct)
-{
-	UScriptStruct* NonConstStruct = const_cast<UScriptStruct*>(GetScriptStruct());
-
-	if (InScriptStruct != NonConstStruct)
-	{
-		InitializeAs(InScriptStruct, nullptr);
-		NonConstStruct = const_cast<UScriptStruct*>(GetScriptStruct());
-	}
-
-	return NonConstStruct;
-}
-
 bool FInstancedStruct::Identical(const FInstancedStruct* Other, uint32 PortFlags) const
 {
 	const UScriptStruct* StructTypePtr = GetScriptStruct();
diff --git a/Engine/Plugins/Experimental/StructUtils/Source/StructUtils/Public/InstancedStruct.h b/Engine/Plugins/Experimental/StructUtils/Source/StructUtils/Public/InstancedStruct.h
--- a/Engine/Plugins/Experimental/StructUtils/Source/StructUtils/Public/InstancedStruct.h
+++ b/Engine/Plugins/Experimental/StructUtils/Source/StructUtils/Public/InstancedStruct.h
@@ -79,13 +79,27 @@ public:
 	{
 		UE::StructUtils::CheckStructType<T>();
 
-		Reset();
-
 		const UScriptStruct* Struct = TBaseStructure<T>::Get();
-		const int32 RequiredSize = Struct->GetStructureSize();
-		uint8* Memory = (uint8*)FMemory::Malloc(FMath::Max(1, RequiredSize));
-		SetStructData(Struct, Memory);
+		uint8* Memory = nullptr;
+
+		const UScriptStruct* CurrentScriptStruct = GetScriptStruct();
+		if (Struct == CurrentScriptStruct)
+		{
+			// Struct type already matches; return the struct memory to a destroyed state so we can placement new over it
+			Memory = GetMutableMemory();
+			CurrentScriptStruct->DestroyStruct(Memory);
+		}
+		else
+		{
+			// Struct type mismatch; reset and reinitialize
+			Reset();
 
+			const int32 RequiredSize = Struct->GetStructureSize();
+			Memory = (uint8*)FMemory::Malloc(FMath::Max(1, RequiredSize));
+			SetStructData(Struct, Memory);
+		}
+
+		check(Memory);
 		new (Memory) T(Forward<TArgs>(InArgs)...);
 	}
 
@@ -239,9 +253,6 @@ public:
 
 protected:
 
-	/** Initializes for new struct type (does nothing if same type) and returns mutable struct. */
-	UScriptStruct* ReinitializeAs(const UScriptStruct* InScriptStruct);
-
 	void DestroyScriptStruct() const
 	{
 		check(StructMemory != nullptr);
