Updated from code review feedback

#rnx

diff --git a/Engine/Plugins/Experimental/StructUtils/Source/StructUtils/Private/InstancedStruct.cpp b/Engine/Plugins/Experimental/StructUtils/Source/StructUtils/Private/InstancedStruct.cpp
--- a/Engine/Plugins/Experimental/StructUtils/Source/StructUtils/Private/InstancedStruct.cpp
+++ b/Engine/Plugins/Experimental/StructUtils/Source/StructUtils/Private/InstancedStruct.cpp
@@ -69,8 +69,7 @@ void FInstancedStruct::InitializeAs(const UScriptStruct* InScriptStruct, const u
 		else
 		{
 			// ... return the struct to its default state
-			CurrentScriptStruct->DestroyStruct(GetMutableMemory());
-			CurrentScriptStruct->InitializeStruct(GetMutableMemory());
+			CurrentScriptStruct->ClearScriptStruct(GetMutableMemory());
 		}
 	}
 	else
@@ -81,8 +80,9 @@ void FInstancedStruct::InitializeAs(const UScriptStruct* InScriptStruct, const u
 		// InScriptStruct == nullptr signifies an empty, unset FInstancedStruct instance
 		if (InScriptStruct)
 		{
+			const int32 MinAlignment = InScriptStruct->GetMinAlignment();
 			const int32 RequiredSize = InScriptStruct->GetStructureSize();
-			const uint8* Memory = ((uint8*)FMemory::Malloc(FMath::Max(1, RequiredSize)));
+			const uint8* Memory = ((uint8*)FMemory::Malloc(FMath::Max(1, RequiredSize), MinAlignment));
 			SetStructData(InScriptStruct, Memory);
 
 			InScriptStruct->InitializeStruct(GetMutableMemory());
diff --git a/Engine/Plugins/Experimental/StructUtils/Source/StructUtils/Public/InstancedStruct.h b/Engine/Plugins/Experimental/StructUtils/Source/StructUtils/Public/InstancedStruct.h
--- a/Engine/Plugins/Experimental/StructUtils/Source/StructUtils/Public/InstancedStruct.h
+++ b/Engine/Plugins/Experimental/StructUtils/Source/StructUtils/Public/InstancedStruct.h
@@ -87,15 +87,16 @@ public:
 		{
 			// Struct type already matches; return the struct memory to a destroyed state so we can placement new over it
 			Memory = GetMutableMemory();
-			CurrentScriptStruct->DestroyStruct(Memory);
+			((T*)Memory)->~T();
 		}
 		else
 		{
 			// Struct type mismatch; reset and reinitialize
 			Reset();
 
+			const int32 MinAlignment = Struct->GetMinAlignment();
 			const int32 RequiredSize = Struct->GetStructureSize();
-			Memory = (uint8*)FMemory::Malloc(FMath::Max(1, RequiredSize));
+			Memory = (uint8*)FMemory::Malloc(FMath::Max(1, RequiredSize), MinAlignment);
 			SetStructData(Struct, Memory);
 		}
 
