Added a new way to construct ScriptStructTypeBitset for types known at compile time.
Also implemented the |-operator
The CL contains tests for the new functionality.

diff --git a/Engine/Plugins/Experimental/StructUtils/Source/StructUtils/Public/ScriptStructTypeBitSet.h b/Engine/Plugins/Experimental/StructUtils/Source/StructUtils/Public/ScriptStructTypeBitSet.h
--- a/Engine/Plugins/Experimental/StructUtils/Source/StructUtils/Public/ScriptStructTypeBitSet.h
+++ b/Engine/Plugins/Experimental/StructUtils/Source/StructUtils/Public/ScriptStructTypeBitSet.h
@@ -281,6 +281,11 @@ private:
 	{
 	}
 
+	TScriptStructTypeBitSet(const int32 BitToSet)
+	{
+		StructTypesBitArray.AddAtIndex(BitToSet);
+	}
+
 public:
 
 	static int32 CreateTypeIndex(const UScriptStruct& InStructType)
@@ -302,6 +307,14 @@ public:
 		return TypeIndex;
 	}
 
+	template<typename T>
+	static const TScriptStructTypeBitSet<TBaseStruct>& GetTypeBitSet()
+	{
+		static_assert(TIsDerivedFrom<T, TBaseStruct>::IsDerived, "Given struct type doesn't match the expected base struct type.");
+		static const TScriptStructTypeBitSet<TBaseStruct> TypeBitSet(GetTypeIndex<T>());
+		return TypeBitSet;
+	}
+	
 	template<typename T>
 	FORCEINLINE void Add()
 	{
@@ -412,6 +425,11 @@ public:
 		return TScriptStructTypeBitSet(TBitArray<>::BitwiseAND(StructTypesBitArray, Other.StructTypesBitArray, EBitwiseOperatorFlags::MinSize));
 	}
 
+	FORCEINLINE TScriptStructTypeBitSet operator|(const TScriptStructTypeBitSet& Other) const
+	{
+		return TScriptStructTypeBitSet(TBitArray<>::BitwiseOR(StructTypesBitArray, Other.StructTypesBitArray, EBitwiseOperatorFlags::MaxSize));
+	}
+
 	FORCEINLINE TScriptStructTypeBitSet GetOverlap(const TScriptStructTypeBitSet& Other) const
 	{
 		return *this & Other;
diff --git a/Engine/Plugins/Experimental/StructUtils/Source/StructUtilsTestSuite/Private/ScriptStructTypeBitSetTest.cpp b/Engine/Plugins/Experimental/StructUtils/Source/StructUtilsTestSuite/Private/ScriptStructTypeBitSetTest.cpp
--- a/Engine/Plugins/Experimental/StructUtils/Source/StructUtilsTestSuite/Private/ScriptStructTypeBitSetTest.cpp
+++ b/Engine/Plugins/Experimental/StructUtils/Source/StructUtilsTestSuite/Private/ScriptStructTypeBitSetTest.cpp
@@ -14,6 +14,8 @@ PRAGMA_DISABLE_OPTIMIZATION
 namespace FScriptStructTypeBitSetTests
 {
 
+using FTestStructBaseBitSet = TScriptStructTypeBitSet<FTestStructSimpleBase>;
+
 struct FTestStructBitSet : public TScriptStructTypeBitSet<FTestStructSimple>
 {
 	FTestStructBitSet() = default;
@@ -243,6 +245,81 @@ struct FStructUtilsTest_BitSetHash : FAITestBase
 
 IMPLEMENT_AI_INSTANT_TEST(FStructUtilsTest_BitSetHash, "System.StructUtils.BitSet.Hash");
 
+struct FStructUtilsTest_CompileTimeBitSet : FAITestBase
+{
+	virtual bool InstantTest() override
+	{
+		FTestStructBaseBitSet::DebugResetStructTypeMappingInfo();
+
+		FTestStructBaseBitSet BitSet1 = FTestStructBaseBitSet::GetTypeBitSet<FTestStructSimple1>();
+		FTestStructBaseBitSet BitSet3 = FTestStructBaseBitSet::GetTypeBitSet<FTestStructSimple3>();
+		FTestStructBaseBitSet BitSet5 = FTestStructBaseBitSet::GetTypeBitSet<FTestStructSimple5>();
+		FTestStructBaseBitSet BitSet4 = FTestStructBaseBitSet::GetTypeBitSet<FTestStructSimple4>();
+		FTestStructBaseBitSet BitSet2 = FTestStructBaseBitSet::GetTypeBitSet<FTestStructSimple2>();
+		FTestStructBaseBitSet BitSet6 = FTestStructBaseBitSet::GetTypeBitSet<FTestStructSimple6>();
+		FTestStructBaseBitSet BitSet7 = FTestStructBaseBitSet::GetTypeBitSet<FTestStructSimple7>();
+
+		AITEST_FALSE("BitSets created for different types should not be equivalent", BitSet1.IsEquivalent(BitSet2));
+		AITEST_FALSE("BitSets created for different types should not be equivalent", BitSet1.IsEquivalent(BitSet5));
+		AITEST_FALSE("BitSets created for different types should not be equivalent", BitSet2.IsEquivalent(BitSet1));
+		AITEST_FALSE("BitSets created for different types should not be equivalent", BitSet2.IsEquivalent(BitSet5));
+		AITEST_FALSE("BitSets created for different types should not be equivalent", BitSet5.IsEquivalent(BitSet1));
+		AITEST_FALSE("BitSets created for different types should not be equivalent", BitSet5.IsEquivalent(BitSet2));
+		
+		FTestStructBaseBitSet BitSet15;
+		BitSet15.Add<FTestStructSimple1>();
+		BitSet15.Add<FTestStructSimple5>();
+		AITEST_TRUE("BitSets created for a list of types should be equivalent to the combination of BitSets to individual types", BitSet15.IsEquivalent(BitSet1 | BitSet5));
+		AITEST_FALSE("BitSets created for a list of types should not contain bit for other types", BitSet15.HasAny(BitSet2));
+		AITEST_FALSE("BitSets created for a list of types should not contain bit for other types", BitSet15.HasAny(BitSet3));
+		AITEST_FALSE("BitSets created for a list of types should not contain bit for other types", BitSet15.HasAny(BitSet4));
+		AITEST_FALSE("BitSets created for a list of types should not contain bit for other types", BitSet15.HasAny(BitSet6));
+		AITEST_FALSE("BitSets created for a list of types should not contain bit for other types", BitSet15.HasAny(BitSet7));
+
+		return true;
+	}
+};
+
+IMPLEMENT_AI_INSTANT_TEST(FStructUtilsTest_CompileTimeBitSet, "System.StructUtils.BitSet.CompileTime");
+
+struct FStructUtilsTest_CompileTimeBitSetReversed : FAITestBase
+{
+	virtual bool InstantTest() override
+	{
+		FTestStructBaseBitSet::DebugResetStructTypeMappingInfo();
+
+		FTestStructBaseBitSet BitSet15;
+		BitSet15.Add<FTestStructSimple1>();
+		BitSet15.Add<FTestStructSimple5>();
+
+		FTestStructBaseBitSet BitSet1 = FTestStructBaseBitSet::GetTypeBitSet<FTestStructSimple1>();
+		FTestStructBaseBitSet BitSet3 = FTestStructBaseBitSet::GetTypeBitSet<FTestStructSimple3>();
+		FTestStructBaseBitSet BitSet5 = FTestStructBaseBitSet::GetTypeBitSet<FTestStructSimple5>();
+		FTestStructBaseBitSet BitSet4 = FTestStructBaseBitSet::GetTypeBitSet<FTestStructSimple4>();
+		FTestStructBaseBitSet BitSet2 = FTestStructBaseBitSet::GetTypeBitSet<FTestStructSimple2>();
+		FTestStructBaseBitSet BitSet6 = FTestStructBaseBitSet::GetTypeBitSet<FTestStructSimple6>();
+		FTestStructBaseBitSet BitSet7 = FTestStructBaseBitSet::GetTypeBitSet<FTestStructSimple7>();
+
+		AITEST_FALSE("BitSets created for different types should not be equivalent", BitSet1.IsEquivalent(BitSet2));
+		AITEST_FALSE("BitSets created for different types should not be equivalent", BitSet1.IsEquivalent(BitSet5));
+		AITEST_FALSE("BitSets created for different types should not be equivalent", BitSet2.IsEquivalent(BitSet1));
+		AITEST_FALSE("BitSets created for different types should not be equivalent", BitSet2.IsEquivalent(BitSet5));
+		AITEST_FALSE("BitSets created for different types should not be equivalent", BitSet5.IsEquivalent(BitSet1));
+		AITEST_FALSE("BitSets created for different types should not be equivalent", BitSet5.IsEquivalent(BitSet2));
+		
+		AITEST_TRUE("BitSets created for a list of types should be equivalent to the combination of BitSets to individual types", BitSet15.IsEquivalent(BitSet1 | BitSet5));
+		AITEST_FALSE("BitSets created for a list of types should not contain bit for other types", BitSet15.HasAny(BitSet2));
+		AITEST_FALSE("BitSets created for a list of types should not contain bit for other types", BitSet15.HasAny(BitSet3));
+		AITEST_FALSE("BitSets created for a list of types should not contain bit for other types", BitSet15.HasAny(BitSet4));
+		AITEST_FALSE("BitSets created for a list of types should not contain bit for other types", BitSet15.HasAny(BitSet6));
+		AITEST_FALSE("BitSets created for a list of types should not contain bit for other types", BitSet15.HasAny(BitSet7));
+
+		return true;
+	}
+};
+
+IMPLEMENT_AI_INSTANT_TEST(FStructUtilsTest_CompileTimeBitSetReversed, "System.StructUtils.BitSet.CompileTime_Reversed");
+
 } // namespace FScriptStructTypeBitSetTests
 
 #undef LOCTEXT_NAMESPACE
diff --git a/Engine/Plugins/Experimental/StructUtils/Source/StructUtilsTestSuite/Public/StructUtilsTestTypes.h b/Engine/Plugins/Experimental/StructUtils/Source/StructUtilsTestSuite/Public/StructUtilsTestTypes.h
--- a/Engine/Plugins/Experimental/StructUtils/Source/StructUtilsTestSuite/Public/StructUtilsTestTypes.h
+++ b/Engine/Plugins/Experimental/StructUtils/Source/StructUtilsTestSuite/Public/StructUtilsTestTypes.h
@@ -6,7 +6,13 @@
 #include "StructUtilsTestTypes.generated.h"
 
 USTRUCT()
-struct FTestStructSimple
+struct FTestStructSimpleBase
+{
+	GENERATED_BODY()
+};
+
+USTRUCT()
+struct FTestStructSimple : public FTestStructSimpleBase
 {
 	GENERATED_BODY()
 	
@@ -32,9 +38,8 @@ struct FTestStructComplex
 	TArray<FString> StringArray;
 };
 
-
 USTRUCT()
-struct FTestStructSimple1
+struct FTestStructSimple1 : public FTestStructSimpleBase
 {
 	GENERATED_BODY()
 	
@@ -45,7 +50,7 @@ struct FTestStructSimple1
 };
 
 USTRUCT()
-struct FTestStructSimple2
+struct FTestStructSimple2 : public FTestStructSimpleBase
 {
 	GENERATED_BODY()
 	
@@ -56,7 +61,7 @@ struct FTestStructSimple2
 };
 
 USTRUCT()
-struct FTestStructSimple3
+struct FTestStructSimple3 : public FTestStructSimpleBase
 {
 	GENERATED_BODY()
 	
@@ -67,7 +72,7 @@ struct FTestStructSimple3
 };
 
 USTRUCT()
-struct FTestStructSimple4
+struct FTestStructSimple4 : public FTestStructSimpleBase
 {
 	GENERATED_BODY()
 	
@@ -78,7 +83,7 @@ struct FTestStructSimple4
 };
 
 USTRUCT()
-struct FTestStructSimple5
+struct FTestStructSimple5 : public FTestStructSimpleBase
 {
 	GENERATED_BODY()
 	
@@ -89,7 +94,7 @@ struct FTestStructSimple5
 };
 
 USTRUCT()
-struct FTestStructSimple6
+struct FTestStructSimple6 : public FTestStructSimpleBase
 {
 	GENERATED_BODY()
 	
@@ -100,7 +105,7 @@ struct FTestStructSimple6
 };
 
 USTRUCT()
-struct FTestStructSimple7
+struct FTestStructSimple7 : public FTestStructSimpleBase
 {
 	GENERATED_BODY()
 	
