Transformed the T::StaticStruct() calls into TBaseStructure<T>::Get() which includes based struct types.

diff --git a/Engine/Plugins/Experimental/StructUtils/Source/StructUtils/Public/ChunkedStructBuffer.h b/Engine/Plugins/Experimental/StructUtils/Source/StructUtils/Public/ChunkedStructBuffer.h
--- a/Engine/Plugins/Experimental/StructUtils/Source/StructUtils/Public/ChunkedStructBuffer.h
+++ b/Engine/Plugins/Experimental/StructUtils/Source/StructUtils/Public/ChunkedStructBuffer.h
@@ -89,7 +89,7 @@ struct STRUCTUTILS_API FChunkedStructBuffer
 	template<typename T, typename... TArgs>
 	void Emplace(TArgs&&... InArgs)
 	{
-		const UScriptStruct* InScriptStruct = T::StaticStruct();
+		const UScriptStruct* InScriptStruct = TBaseStructure<T>::Get();
 		uint8* ItemMemory = AllocItem(InScriptStruct);
 		new (ItemMemory) T(Forward<TArgs>(InArgs)...);
 	}
@@ -98,7 +98,7 @@ struct STRUCTUTILS_API FChunkedStructBuffer
 	template<typename T, typename... TArgs>
 	T& Emplace_GetRef(TArgs&&... InArgs)
 	{
-		const UScriptStruct* InScriptStruct = T::StaticStruct();
+		const UScriptStruct* InScriptStruct = TBaseStructure<T>::Get();
 		uint8* ItemMemory = AllocItem(InScriptStruct);
 		return *new (ItemMemory) T(Forward<TArgs>(InArgs)...);
 	}
@@ -107,7 +107,7 @@ struct STRUCTUTILS_API FChunkedStructBuffer
 	template<typename T>
 	void Add(const T& InStruct)
 	{
-		const UScriptStruct* InScriptStruct = T::StaticStruct();
+		const UScriptStruct* InScriptStruct = TBaseStructure<T>::Get();
 		uint8* ItemMemory = AllocItem(InScriptStruct);
 		new (ItemMemory) T(InStruct);
 	}
@@ -116,7 +116,7 @@ struct STRUCTUTILS_API FChunkedStructBuffer
 	template<typename T>
 	T& Add_GetRef(const T& InStruct)
 	{
-		const UScriptStruct* InScriptStruct = T::StaticStruct();
+		const UScriptStruct* InScriptStruct = TBaseStructure<T>::Get();
 		uint8* ItemMemory = AllocItem(InScriptStruct);
 		return *new (ItemMemory) T(InStruct);
 	}
@@ -160,7 +160,7 @@ struct STRUCTUTILS_API FChunkedStructBuffer
 		for (FChunkHeader* Chunk = First; Chunk != nullptr; Chunk = Chunk->Next)
 		{
 			// Build mask of the accepted struct in the chunk.
-			const int32 StructIndex = Chunk->GetScriptStructIndex(T::StaticStruct());
+			const int32 StructIndex = Chunk->GetScriptStructIndex(TBaseStructure<T>::Get());
 			if (StructIndex == INDEX_NONE)
 			{
 				continue;
diff --git a/Engine/Plugins/Experimental/StructUtils/Source/StructUtils/Public/InstancedStructStream.h b/Engine/Plugins/Experimental/StructUtils/Source/StructUtils/Public/InstancedStructStream.h
--- a/Engine/Plugins/Experimental/StructUtils/Source/StructUtils/Public/InstancedStructStream.h
+++ b/Engine/Plugins/Experimental/StructUtils/Source/StructUtils/Public/InstancedStructStream.h
@@ -107,7 +107,7 @@ struct STRUCTUTILS_API FInstancedStructStream
 	template<typename T, typename... TArgs>
 	void Emplace(TArgs&&... InArgs)
 	{
-		const UScriptStruct* InScriptStruct = T::StaticStruct();
+		const UScriptStruct* InScriptStruct = TBaseStructure<T>::Get();
 		uint8* ItemMemory = AllocItem(InScriptStruct);
 		new (ItemMemory) T(Forward<TArgs>(InArgs)...);
 	}
@@ -116,7 +116,7 @@ struct STRUCTUTILS_API FInstancedStructStream
 	template<typename T, typename... TArgs>
 	T& Emplace_GetRef(TArgs&&... InArgs)
 	{
-		const UScriptStruct* InScriptStruct = T::StaticStruct();
+		const UScriptStruct* InScriptStruct = TBaseStructure<T>::Get();
 		uint8* ItemMemory = AllocItem(InScriptStruct);
 		return *new (ItemMemory) T(Forward<TArgs>(InArgs)...);
 	}
@@ -125,7 +125,7 @@ struct STRUCTUTILS_API FInstancedStructStream
 	template<typename T>
 	void Add(const T& InStruct)
 	{
-		const UScriptStruct* InScriptStruct = T::StaticStruct();
+		const UScriptStruct* InScriptStruct = TBaseStructure<T>::Get();
 		uint8* ItemMemory = AllocItem(InScriptStruct);
 		new (ItemMemory) T(InStruct);
 	}
@@ -134,7 +134,7 @@ struct STRUCTUTILS_API FInstancedStructStream
 	template<typename T>
 	T& Add_GetRef(const T& InStruct)
 	{
-		const UScriptStruct* InScriptStruct = T::StaticStruct();
+		const UScriptStruct* InScriptStruct = TBaseStructure<T>::Get();
 		uint8* ItemMemory = AllocItem(InScriptStruct);
 		return *new (ItemMemory) T(InStruct);
 	}
@@ -178,7 +178,7 @@ struct STRUCTUTILS_API FInstancedStructStream
 		for (FChunkHeader* Chunk = First; Chunk != nullptr; Chunk = Chunk->Next)
 		{
 			// Build mask of the accepted struct in the chunk.
-			const int32 StructIndex = Chunk->GetScriptStructIndex(T::StaticStruct());
+			const int32 StructIndex = Chunk->GetScriptStructIndex(TBaseStructure<T>::Get());
 			if (StructIndex == INDEX_NONE)
 			{
 				continue;
diff --git a/Engine/Plugins/Experimental/StructUtils/Source/StructUtils/Public/ScriptStructTypeBitSet.h b/Engine/Plugins/Experimental/StructUtils/Source/StructUtils/Public/ScriptStructTypeBitSet.h
--- a/Engine/Plugins/Experimental/StructUtils/Source/StructUtils/Public/ScriptStructTypeBitSet.h
+++ b/Engine/Plugins/Experimental/StructUtils/Source/StructUtils/Public/ScriptStructTypeBitSet.h
@@ -303,7 +303,7 @@ public:
 	static int32 GetTypeIndex()
 	{
 		static_assert(TIsDerivedFrom<T, TBaseStruct>::IsDerived, "Given struct type doesn't match the expected base struct type.");
-		static const int32 TypeIndex = CreateTypeIndex(*T::StaticStruct());
+		static const int32 TypeIndex = CreateTypeIndex(*TBaseStructure<T>::Get());
 		return TypeIndex;
 	}
 
diff --git a/Engine/Plugins/Experimental/StructUtils/Source/StructUtils/Public/SharedStruct.h b/Engine/Plugins/Experimental/StructUtils/Source/StructUtils/Public/SharedStruct.h
--- a/Engine/Plugins/Experimental/StructUtils/Source/StructUtils/Public/SharedStruct.h
+++ b/Engine/Plugins/Experimental/StructUtils/Source/StructUtils/Public/SharedStruct.h
@@ -127,7 +127,7 @@ struct STRUCTUTILS_API FConstSharedStruct
 		const UScriptStruct* Struct = GetScriptStruct();
 		check(Memory != nullptr);
 		check(Struct != nullptr);
-		check(Struct->IsChildOf(T::StaticStruct()));
+		check(Struct->IsChildOf(TBaseStructure<T>::Get()));
 		return *((T*)Memory);
 	}
 
@@ -137,7 +137,7 @@ struct STRUCTUTILS_API FConstSharedStruct
 	{
 		const uint8* Memory = GetMemory();
 		const UScriptStruct* Struct = GetScriptStruct();
-		if (Memory != nullptr && Struct && Struct->IsChildOf(T::StaticStruct()))
+		if (Memory != nullptr && Struct && Struct->IsChildOf(TBaseStructure<T>::Get()))
 		{
 			return ((T*)Memory);
 		}
@@ -242,7 +242,7 @@ struct STRUCTUTILS_API FSharedStruct : public FConstSharedStruct
 		UE::StructUtils::CheckStructType<T>();
 
 		Reset();
-		StructMemoryPtr = FStructSharedMemory::Create(*T::StaticStruct());
+		StructMemoryPtr = FStructSharedMemory::Create(*TBaseStructure<T>::Get());
 		new (GetMutableMemory()) T(Forward<TArgs>(InArgs)...);
 	}
 
@@ -253,7 +253,7 @@ struct STRUCTUTILS_API FSharedStruct : public FConstSharedStruct
 		UE::StructUtils::CheckStructType<T>();
 
 		FSharedStruct SharedStruct;
-		SharedStruct.InitializeAs(T::StaticStruct(), nullptr);
+		SharedStruct.InitializeAs(TBaseStructure<T>::Get(), nullptr);
 		return SharedStruct;
 	}
 
@@ -264,7 +264,7 @@ struct STRUCTUTILS_API FSharedStruct : public FConstSharedStruct
 		UE::StructUtils::CheckStructType<T>();
 
 		FSharedStruct SharedStruct;
-		SharedStruct.InitializeAs(T::StaticStruct(), reinterpret_cast<const uint8*>(&Struct));
+		SharedStruct.InitializeAs(TBaseStructure<T>::Get(), reinterpret_cast<const uint8*>(&Struct));
 		return SharedStruct;
 	}
 
@@ -294,7 +294,7 @@ struct STRUCTUTILS_API FSharedStruct : public FConstSharedStruct
 		const UScriptStruct* Struct = GetScriptStruct();
 		check(Memory != nullptr);
 		check(Struct != nullptr);
-		check(Struct->IsChildOf(T::StaticStruct()));
+		check(Struct->IsChildOf(TBaseStructure<T>::Get()));
 		return *((T*)Memory);
 	}
 
@@ -304,7 +304,7 @@ struct STRUCTUTILS_API FSharedStruct : public FConstSharedStruct
 	{
 		uint8* Memory = GetMutableMemory();
 		const UScriptStruct* Struct = GetScriptStruct();
-		if (Memory != nullptr && Struct && Struct->IsChildOf(T::StaticStruct()))
+		if (Memory != nullptr && Struct && Struct->IsChildOf(TBaseStructure<T>::Get()))
 		{
 			return ((T*)Memory);
 		}
diff --git a/Engine/Plugins/Experimental/StructUtils/Source/StructUtils/Public/StructArrayView.h b/Engine/Plugins/Experimental/StructUtils/Source/StructUtils/Public/StructArrayView.h
--- a/Engine/Plugins/Experimental/StructUtils/Source/StructUtils/Public/StructArrayView.h
+++ b/Engine/Plugins/Experimental/StructUtils/Source/StructUtils/Public/StructArrayView.h
@@ -73,7 +73,7 @@ struct FStructArrayView
 	template<typename T>
 	const T& GetElementAtChecked(const int32 Index) const
 	{
-		check(T::StaticStruct() == &FragmentType);
+		check(TBaseStructure<T>::Get() == &FragmentType);
 		return *((T*)GetDataAt(Index));
 	}
 
@@ -86,7 +86,7 @@ struct FStructArrayView
 	template<typename T>
 	T& GetMutableElementAtChecked(const int32 Index)
 	{
-		check(T::StaticStruct() == &FragmentType);
+		check(TBaseStructure<T>::Get() == &FragmentType);
 		return *((T*)GetDataAt(Index));
 	}
 
