InstancedStruct - When cloning the instanced struct for editing, make sure to initialize the editable struct to match the source data.

diff --git a/Engine/Plugins/Experimental/StructUtils/Source/StructUtilsEditor/Private/Customizations/InstancedStructDetails.cpp b/Engine/Plugins/Experimental/StructUtils/Source/StructUtilsEditor/Private/Customizations/InstancedStructDetails.cpp
--- a/Engine/Plugins/Experimental/StructUtils/Source/StructUtilsEditor/Private/Customizations/InstancedStructDetails.cpp
+++ b/Engine/Plugins/Experimental/StructUtils/Source/StructUtilsEditor/Private/Customizations/InstancedStructDetails.cpp
@@ -238,6 +238,9 @@ void FInstancedStructDataDetails::GenerateChildContent(IDetailChildrenBuilder& C
 				StructInstanceData->SetPackage(OuterPackages[0]);
 			}
 		}
+
+		bool bStructMismatch = false;
+		SyncEditableInstanceFromSource(&bStructMismatch);
 	}
 
 	// Add the rows for the struct
