[Backout] - CL19681262
Original CL Desc
-----------------------------------------------------------------
Allow AddChildStructure to be used by sparse data properties

The child struct properties can't inherit the sparse data flag, as FStructurePropertyNode is already directly editing a block of struct memory.

diff --git a/Engine/Plugins/Experimental/StructUtils/Source/StructUtilsEditor/Private/Customizations/InstancedStructDetails.cpp b/Engine/Plugins/Experimental/StructUtils/Source/StructUtilsEditor/Private/Customizations/InstancedStructDetails.cpp
--- a/Engine/Plugins/Experimental/StructUtils/Source/StructUtilsEditor/Private/Customizations/InstancedStructDetails.cpp
+++ b/Engine/Plugins/Experimental/StructUtils/Source/StructUtilsEditor/Private/Customizations/InstancedStructDetails.cpp
@@ -249,16 +249,20 @@ void FInstancedStructDataDetails::GenerateChildContent(IDetailChildrenBuilder& C
 		FSimpleDelegate OnStructValuePreChangeDelegate = FSimpleDelegate::CreateSP(this, &FInstancedStructDataDetails::OnStructValuePreChange);
 		FSimpleDelegate OnStructValuePostChangeDelegate = FSimpleDelegate::CreateSP(this, &FInstancedStructDataDetails::OnStructValuePostChange);
 
-		TArray<TSharedPtr<IPropertyHandle>> ChildProperties = StructProperty->AddChildStructure(StructInstanceData.ToSharedRef());
-		for (TSharedPtr<IPropertyHandle> ChildHandle : ChildProperties)
+		// Note: We use AddExternalStructureProperty here as it gives the desired result (the struct value properties as direct children of the struct header)
+		// Neither AddExternalStructure (which added an extra row) nor AddAllExternalStructureProperties (which didn't indent the properties as children) were suitable
+		for (TFieldIterator<const FProperty> PropertyIt(StructInstanceData->GetStruct()); PropertyIt; ++PropertyIt)
 		{
-			ChildHandle->SetOnPropertyValuePreChange(OnStructValuePreChangeDelegate);
-			ChildHandle->SetOnChildPropertyValuePreChange(OnStructValuePreChangeDelegate);
-			ChildHandle->SetOnPropertyValueChanged(OnStructValuePostChangeDelegate);
-			ChildHandle->SetOnChildPropertyValueChanged(OnStructValuePostChangeDelegate);
+			if (IDetailPropertyRow* StructValuePropertyRow = ChildBuilder.AddExternalStructureProperty(StructInstanceData.ToSharedRef(), PropertyIt->GetFName()))
+			{
+				TSharedPtr<IPropertyHandle> StructValuePropertyHandle = StructValuePropertyRow->GetPropertyHandle();
+				StructValuePropertyHandle->SetOnPropertyValuePreChange(OnStructValuePreChangeDelegate);
+				StructValuePropertyHandle->SetOnChildPropertyValuePreChange(OnStructValuePreChangeDelegate);
+				StructValuePropertyHandle->SetOnPropertyValueChanged(OnStructValuePostChangeDelegate);
+				StructValuePropertyHandle->SetOnChildPropertyValueChanged(OnStructValuePostChangeDelegate);
 
-			IDetailPropertyRow& Row = ChildBuilder.AddProperty(ChildHandle.ToSharedRef());
-			OnChildRowAdded(Row);
+				OnChildRowAdded(*StructValuePropertyRow);
+			}
 		}
 	}
 }
diff --git a/Engine/Source/Editor/PropertyEditor/Private/CategoryPropertyNode.cpp b/Engine/Source/Editor/PropertyEditor/Private/CategoryPropertyNode.cpp
--- a/Engine/Source/Editor/PropertyEditor/Private/CategoryPropertyNode.cpp
+++ b/Engine/Source/Editor/PropertyEditor/Private/CategoryPropertyNode.cpp
@@ -101,7 +101,7 @@ void FCategoryPropertyNode::InitChildNodes()
 		InitParams.bAllowChildren = true;
 		InitParams.bForceHiddenPropertyVisibility = bShowHiddenProperties;
 		InitParams.bCreateDisableEditOnInstanceNodes = bShouldShowDisableEditOnInstance;
-		InitParams.IsSparseProperty = SparseProperties.Contains(Properties[PropertyIndex]) ? FPropertyNodeInitParams::EIsSparseDataProperty::True : FPropertyNodeInitParams::EIsSparseDataProperty::Inherit;
+		InitParams.bIsSparseProperty = SparseProperties.Contains(Properties[PropertyIndex]);
 
 		NewItemNode->InitNode( InitParams );
 
diff --git a/Engine/Source/Editor/PropertyEditor/Private/PropertyHandleImpl.cpp b/Engine/Source/Editor/PropertyEditor/Private/PropertyHandleImpl.cpp
--- a/Engine/Source/Editor/PropertyEditor/Private/PropertyHandleImpl.cpp
+++ b/Engine/Source/Editor/PropertyEditor/Private/PropertyHandleImpl.cpp
@@ -3135,7 +3135,6 @@ TArray<TSharedPtr<IPropertyHandle>> FPropertyHandleBase::AddChildStructure( TSha
 	RootInitParams.bAllowChildren = true;
 	RootInitParams.bForceHiddenPropertyVisibility = FPropertySettings::Get().ShowHiddenProperties();
 	RootInitParams.bCreateCategoryNodes = false;
-	RootInitParams.IsSparseProperty = FPropertyNodeInitParams::EIsSparseDataProperty::False; // FStructurePropertyNode can't inherit the sparse data flag
 
 	StructPropertyNode->InitNode(RootInitParams);
 
diff --git a/Engine/Source/Editor/PropertyEditor/Private/PropertyNode.cpp b/Engine/Source/Editor/PropertyEditor/Private/PropertyNode.cpp
--- a/Engine/Source/Editor/PropertyEditor/Private/PropertyNode.cpp
+++ b/Engine/Source/Editor/PropertyEditor/Private/PropertyNode.cpp
@@ -100,9 +100,9 @@ void FPropertyNode::InitNode(const FPropertyNodeInitParams& InitParams)
 	ArrayOffset = InitParams.ArrayOffset;
 	ArrayIndex = InitParams.ArrayIndex;
 
-	bool bIsSparse = InitParams.IsSparseProperty == FPropertyNodeInitParams::EIsSparseDataProperty::True;
+	bool bIsSparse = InitParams.bIsSparseProperty;
 
-	if (ParentNode && InitParams.IsSparseProperty == FPropertyNodeInitParams::EIsSparseDataProperty::Inherit)
+	if (ParentNode)
 	{
 		//default to parents max child depth
 		MaxChildDepthAllowed = ParentNode->MaxChildDepthAllowed;
diff --git a/Engine/Source/Editor/PropertyEditor/Private/PropertyNode.h b/Engine/Source/Editor/PropertyEditor/Private/PropertyNode.h
--- a/Engine/Source/Editor/PropertyEditor/Private/PropertyNode.h
+++ b/Engine/Source/Editor/PropertyEditor/Private/PropertyNode.h
@@ -198,13 +198,6 @@ private:
  */
 struct FPropertyNodeInitParams
 {
-	enum class EIsSparseDataProperty : uint8
-	{
-		False,
-		True,
-		Inherit,
-	};
-
 	/** The parent of the property node */
 	TSharedPtr<FPropertyNode> ParentNode;
 	/** The property that the node observes and modifies*/
@@ -222,7 +215,7 @@ struct FPropertyNodeInitParams
 	/** Whether or not to create nodes for properties marked CPF_DisableEditOnInstance */
 	bool bCreateDisableEditOnInstanceNodes;
 	/** Whether or not this property is sparse data */
-	EIsSparseDataProperty IsSparseProperty;
+	bool bIsSparseProperty;
 
 	FPropertyNodeInitParams()
 		: ParentNode(nullptr)
@@ -233,7 +226,7 @@ struct FPropertyNodeInitParams
 		, bForceHiddenPropertyVisibility( false )
 		, bCreateCategoryNodes( true )
 		, bCreateDisableEditOnInstanceNodes( true )
-		, IsSparseProperty( EIsSparseDataProperty::Inherit )
+		, bIsSparseProperty( false )
 	{}
 };
 
