Property Bag: UPropertyBag and FInstancedPropertyBag
- Property bags allows to create customs structs via code or UI and store them in an property bag instance
- Bags are identified by their layout, a UScriptStruct is created for each unique layout
- The instances serialize the layout, not the type, to simplify handling with old data
- Property values can be migrated between different bags

diff --git a/Engine/Plugins/Experimental/StructUtils/Source/StructUtils/Private/PropertyBag.cpp b/Engine/Plugins/Experimental/StructUtils/Source/StructUtils/Private/PropertyBag.cpp
new file mode 100644
--- /dev/null
+++ b/Engine/Plugins/Experimental/StructUtils/Source/StructUtils/Private/PropertyBag.cpp
@@ -0,0 +1,1196 @@
+﻿#include "PropertyBag.h"
+#include "StructView.h"
+#include "Hash/CityHash.h"
+
+namespace UE::StructUtils::Private
+{
+
+	bool CanCastTo(const UStruct* From, const UStruct* To)
+	{
+		return From != nullptr && To != nullptr && From->IsChildOf(To);
+	}
+
+	uint64 GetObjectHash(const UObject* Object)
+	{
+		const FString PathName = GetPathNameSafe(Object);
+		return CityHash64((const char*)GetData(PathName), PathName.Len() * sizeof(TCHAR));
+	}
+
+	uint64 CalcPropertyDescHash(const FPropertyBagPropertyDesc& Desc)
+	{
+		const uint32 Hashes[] = { GetTypeHash(Desc.ID), GetTypeHash(Desc.Name), GetTypeHash(Desc.ValueType) };
+		return CityHash64WithSeed((const char*)Hashes, sizeof(Hashes), GetObjectHash(Desc.ValueTypeObject));
+	}
+
+	uint64 CalcPropertyDescArrayHash(const TConstArrayView<FPropertyBagPropertyDesc> Descs)
+	{
+		uint64 Hash = 0;
+		for (const FPropertyBagPropertyDesc& Desc : Descs)
+		{
+			Hash = CityHash128to64(Uint128_64(Hash, CalcPropertyDescHash(Desc)));
+		}
+		return Hash;
+	}
+
+	FProperty* CreatePropertyFromDesc(const FPropertyBagPropertyDesc& Desc, UScriptStruct* PropertyScope)
+	{
+		switch (Desc.ValueType)
+		{
+		case EPropertyBagPropertyType::Bool:
+			{
+				FBoolProperty* Prop = new FBoolProperty(PropertyScope, Desc.Name, RF_Public);
+				return Prop;
+			}
+		case EPropertyBagPropertyType::Byte:
+			{
+				FByteProperty* Prop = new FByteProperty(PropertyScope, Desc.Name, RF_Public);
+				Prop->SetPropertyFlags(CPF_HasGetValueTypeHash);
+				return Prop;
+			}
+		case EPropertyBagPropertyType::Int32:
+			{
+				FIntProperty* Prop = new FIntProperty(PropertyScope, Desc.Name, RF_Public);
+				Prop->SetPropertyFlags(CPF_HasGetValueTypeHash);
+				return Prop;
+			}
+		case EPropertyBagPropertyType::Int64:
+			{
+				FInt64Property* Prop = new FInt64Property(PropertyScope, Desc.Name, RF_Public);
+				Prop->SetPropertyFlags(CPF_HasGetValueTypeHash);
+				return Prop;
+			}
+		case EPropertyBagPropertyType::Float:
+			{
+				FFloatProperty* Prop = new FFloatProperty(PropertyScope, Desc.Name, RF_Public);
+				Prop->SetPropertyFlags(CPF_HasGetValueTypeHash);
+				return Prop;
+			}
+		case EPropertyBagPropertyType::Double:
+			{
+				FDoubleProperty* Prop = new FDoubleProperty(PropertyScope, Desc.Name, RF_Public);
+				Prop->SetPropertyFlags(CPF_HasGetValueTypeHash);
+				return Prop;
+			}
+		case EPropertyBagPropertyType::Name:
+			{
+				FNameProperty* Prop = new FNameProperty(PropertyScope, Desc.Name, RF_Public);
+				Prop->SetPropertyFlags(CPF_HasGetValueTypeHash);
+				return Prop;
+			}
+		case EPropertyBagPropertyType::String:
+			{
+				FStrProperty* Prop = new FStrProperty(PropertyScope, Desc.Name, RF_Public);
+				Prop->SetPropertyFlags(CPF_HasGetValueTypeHash);
+				return Prop;
+			}
+		case EPropertyBagPropertyType::Text:
+			{
+				FTextProperty* Prop = new FTextProperty(PropertyScope, Desc.Name, RF_Public);
+				return Prop;
+			}
+		case EPropertyBagPropertyType::Enum:
+			if (UEnum* Enum = Cast<UEnum>(Desc.ValueTypeObject))
+			{
+				FEnumProperty* Prop = new FEnumProperty(PropertyScope, Desc.Name, RF_Public);
+				FNumericProperty* UnderlyingProp = new FByteProperty(Prop, "UnderlyingType", RF_Public); // HACK: Hardwire to byte property for now for BP compatibility
+				Prop->SetEnum(Enum);
+				Prop->AddCppProperty(UnderlyingProp);
+				return Prop;
+			}
+			break;
+		case EPropertyBagPropertyType::Struct:
+			if (UScriptStruct* ScriptStruct = Cast<UScriptStruct>(Desc.ValueTypeObject))
+			{
+				FStructProperty* Prop = new FStructProperty(PropertyScope, Desc.Name, RF_Public);
+				Prop->Struct = ScriptStruct;
+
+				if (ScriptStruct->GetCppStructOps() && ScriptStruct->GetCppStructOps()->HasGetTypeHash())
+				{
+					Prop->SetPropertyFlags(CPF_HasGetValueTypeHash);
+				}
+				
+				if (ScriptStruct->StructFlags & STRUCT_HasInstancedReference)
+				{
+					Prop->SetPropertyFlags(CPF_ContainsInstancedReference);
+				}
+				
+				return Prop;
+			}
+			break;
+		case EPropertyBagPropertyType::Object:
+			if (UClass* Class = Cast<UClass>(Desc.ValueTypeObject))
+			{
+				FObjectProperty* Prop = new FObjectProperty(PropertyScope, Desc.Name, RF_Public);
+				if (Class->HasAnyClassFlags(CLASS_DefaultToInstanced))
+				{
+					Prop->SetPropertyFlags(CPF_InstancedReference);
+				}
+				Prop->SetPropertyClass(Class);
+				Prop->SetPropertyFlags(CPF_HasGetValueTypeHash);
+				return Prop;
+			}
+			break;
+		case EPropertyBagPropertyType::SoftObject:
+			if (UClass* Class = Cast<UClass>(Desc.ValueTypeObject))
+			{
+				FSoftObjectProperty* Prop = new FSoftObjectProperty(PropertyScope, Desc.Name, RF_Public);
+				if (Class->HasAnyClassFlags(CLASS_DefaultToInstanced))
+				{
+					Prop->SetPropertyFlags(CPF_InstancedReference);
+				}
+				Prop->SetPropertyClass(Class);
+				Prop->SetPropertyFlags(CPF_HasGetValueTypeHash);
+				return Prop;
+			}
+			break;
+		case EPropertyBagPropertyType::Class:
+			if (UClass* Class = Cast<UClass>(Desc.ValueTypeObject))
+			{
+				FClassProperty* Prop = new FClassProperty(PropertyScope, Desc.Name, RF_Public);
+				Prop->SetMetaClass(Class);
+				Prop->PropertyClass = UClass::StaticClass();
+				Prop->SetPropertyFlags(CPF_HasGetValueTypeHash);
+				return Prop;
+			}
+			break;
+		case EPropertyBagPropertyType::SoftClass:
+			if (UClass* Class = Cast<UClass>(Desc.ValueTypeObject))
+			{
+				FSoftClassProperty* Prop = new FSoftClassProperty(PropertyScope, Desc.Name, RF_Public);
+				Prop->SetMetaClass(Class);
+				Prop->PropertyClass = UClass::StaticClass();
+				Prop->SetPropertyFlags(CPF_HasGetValueTypeHash);
+				return Prop;
+			}
+			break;
+		default:
+			ensureMsgf(false, TEXT("Unhandled stype %s"), *UEnum::GetValueAsString(Desc.ValueType));
+		}
+
+		return nullptr;
+	}
+
+	EPropertyBagResult GetPropertyAsDouble(const FPropertyBagPropertyDesc& Desc, const void* Address, double& OutValue)
+	{
+		check(Desc.CachedProperty);
+		check(Address);
+		
+		switch(Desc.ValueType)
+		{
+		case EPropertyBagPropertyType::Bool:
+			{
+				const FBoolProperty* Property = CastFieldChecked<FBoolProperty>(Desc.CachedProperty);
+				OutValue = Property->GetPropertyValue(Address) ? 1.0 : 0.0;
+				return EPropertyBagResult::Success;
+			}
+		case EPropertyBagPropertyType::Byte:
+			{
+				const FByteProperty* Property = CastFieldChecked<FByteProperty>(Desc.CachedProperty);
+				OutValue = Property->GetPropertyValue(Address);
+				return EPropertyBagResult::Success;
+			}
+		case EPropertyBagPropertyType::Int32:
+			{
+				const FIntProperty* Property = CastFieldChecked<FIntProperty>(Desc.CachedProperty);
+				OutValue = Property->GetPropertyValue(Address);
+				return EPropertyBagResult::Success;
+			}
+		case EPropertyBagPropertyType::Int64:
+			{
+				const FInt64Property* Property = CastFieldChecked<FInt64Property>(Desc.CachedProperty);
+				OutValue = Property->GetPropertyValue(Address);
+				return EPropertyBagResult::Success;
+			}
+		case EPropertyBagPropertyType::Float:
+			{
+				const FFloatProperty* Property = CastFieldChecked<FFloatProperty>(Desc.CachedProperty);
+				OutValue = Property->GetPropertyValue(Address);
+				return EPropertyBagResult::Success;
+			}
+		case EPropertyBagPropertyType::Double:
+			{
+				const FDoubleProperty* Property = CastFieldChecked<FDoubleProperty>(Desc.CachedProperty);
+				OutValue = Property->GetPropertyValue(Address);
+				return EPropertyBagResult::Success;
+			}
+		case EPropertyBagPropertyType::Enum:
+			{
+				const FEnumProperty* EnumProperty = CastFieldChecked<FEnumProperty>(Desc.CachedProperty);
+				const FNumericProperty* UnderlyingProperty = EnumProperty->GetUnderlyingProperty();
+				check(UnderlyingProperty);
+				OutValue = UnderlyingProperty->GetSignedIntPropertyValue(Address);
+				return EPropertyBagResult::Success;
+			}
+		default:
+			return EPropertyBagResult::TypeMismatch;
+		}
+	}
+
+	EPropertyBagResult SetPropertyFromDouble(const FPropertyBagPropertyDesc& Desc, void* Address, const double InValue)
+	{
+		check(Desc.CachedProperty);
+		check(Address);
+		
+		switch(Desc.ValueType)
+		{
+		case EPropertyBagPropertyType::Bool:
+			{
+				const FBoolProperty* Property = CastFieldChecked<FBoolProperty>(Desc.CachedProperty);
+				Property->SetPropertyValue(Address, FMath::IsNearlyZero(InValue) ? false : true);
+				return EPropertyBagResult::Success;
+			}
+		case EPropertyBagPropertyType::Byte:
+			{
+				const FByteProperty* Property = CastFieldChecked<FByteProperty>(Desc.CachedProperty);
+				Property->SetPropertyValue(Address, FMath::RoundToInt32(InValue));
+				return EPropertyBagResult::Success;
+			}
+		case EPropertyBagPropertyType::Int32:
+			{
+				const FIntProperty* Property = CastFieldChecked<FIntProperty>(Desc.CachedProperty);
+				Property->SetPropertyValue(Address, FMath::RoundToInt32(InValue));
+				return EPropertyBagResult::Success;
+			}
+		case EPropertyBagPropertyType::Int64:
+			{
+				const FInt64Property* Property = CastFieldChecked<FInt64Property>(Desc.CachedProperty);
+				Property->SetPropertyValue(Address, FMath::RoundToInt64(InValue));
+				return EPropertyBagResult::Success;
+			}
+		case EPropertyBagPropertyType::Float:
+			{
+				const FFloatProperty* Property = CastFieldChecked<FFloatProperty>(Desc.CachedProperty);
+				Property->SetPropertyValue(Address, (float)InValue);
+				return EPropertyBagResult::Success;
+			}
+		case EPropertyBagPropertyType::Double:
+			{
+				const FDoubleProperty* Property = CastFieldChecked<FDoubleProperty>(Desc.CachedProperty);
+				Property->SetPropertyValue(Address, InValue);
+				return EPropertyBagResult::Success;
+			}
+		case EPropertyBagPropertyType::Enum:
+			{
+				const FEnumProperty* EnumProperty = CastFieldChecked<FEnumProperty>(Desc.CachedProperty);
+				const FNumericProperty* UnderlyingProperty = EnumProperty->GetUnderlyingProperty();
+				check(UnderlyingProperty);
+				UnderlyingProperty->SetIntPropertyValue(Address, (uint64)InValue);
+				return EPropertyBagResult::Success;
+			}
+		default:
+			return EPropertyBagResult::TypeMismatch;
+		}
+	}
+
+	EPropertyBagResult GetPropertyAsInt64(const FPropertyBagPropertyDesc& Desc, const void* Address, int64& OutValue)
+	{
+		check(Desc.CachedProperty);
+		check(Address);
+		
+		switch(Desc.ValueType)
+		{
+		case EPropertyBagPropertyType::Bool:
+			{
+				const FBoolProperty* Property = CastFieldChecked<FBoolProperty>(Desc.CachedProperty);
+				OutValue = Property->GetPropertyValue(Address) ? 1 : 0;
+				return EPropertyBagResult::Success;
+			}
+		case EPropertyBagPropertyType::Byte:
+			{
+				const FByteProperty* Property = CastFieldChecked<FByteProperty>(Desc.CachedProperty);
+				OutValue = Property->GetPropertyValue(Address);
+				return EPropertyBagResult::Success;
+			}
+		case EPropertyBagPropertyType::Int32:
+			{
+				const FIntProperty* Property = CastFieldChecked<FIntProperty>(Desc.CachedProperty);
+				OutValue = Property->GetPropertyValue(Address);
+				return EPropertyBagResult::Success;
+			}
+		case EPropertyBagPropertyType::Int64:
+			{
+				const FInt64Property* Property = CastFieldChecked<FInt64Property>(Desc.CachedProperty);
+				OutValue = Property->GetPropertyValue(Address);
+				return EPropertyBagResult::Success;
+			}
+		case EPropertyBagPropertyType::Float:
+			{
+				const FFloatProperty* Property = CastFieldChecked<FFloatProperty>(Desc.CachedProperty);
+				OutValue = (int64)Property->GetPropertyValue(Address);
+				return EPropertyBagResult::Success;
+			}
+		case EPropertyBagPropertyType::Double:
+			{
+				const FDoubleProperty* Property = CastFieldChecked<FDoubleProperty>(Desc.CachedProperty);
+				OutValue = (int64)Property->GetPropertyValue(Address);
+				return EPropertyBagResult::Success;
+			}
+		case EPropertyBagPropertyType::Enum:
+			{
+				const FEnumProperty* EnumProperty = CastFieldChecked<FEnumProperty>(Desc.CachedProperty);
+				const FNumericProperty* UnderlyingProperty = EnumProperty->GetUnderlyingProperty();
+				check(UnderlyingProperty);
+				OutValue = UnderlyingProperty->GetSignedIntPropertyValue(Address);
+				return EPropertyBagResult::Success;
+			}
+		default:
+			return EPropertyBagResult::TypeMismatch;
+		}
+	}
+
+	EPropertyBagResult SetPropertyFromInt64(const FPropertyBagPropertyDesc& Desc, void* Address, const int64 InValue)
+	{
+		check(Desc.CachedProperty);
+		check(Address);
+		
+		switch(Desc.ValueType)
+		{
+		case EPropertyBagPropertyType::Bool:
+			{
+				const FBoolProperty* Property = CastFieldChecked<FBoolProperty>(Desc.CachedProperty);
+				Property->SetPropertyValue(Address, InValue != 0);
+				return EPropertyBagResult::Success;
+			}
+		case EPropertyBagPropertyType::Byte:
+			{
+				const FByteProperty* Property = CastFieldChecked<FByteProperty>(Desc.CachedProperty);
+				Property->SetPropertyValue(Address, (uint8)InValue);
+				return EPropertyBagResult::Success;
+			}
+		case EPropertyBagPropertyType::Int32:
+			{
+				const FIntProperty* Property = CastFieldChecked<FIntProperty>(Desc.CachedProperty);
+				Property->SetPropertyValue(Address, (int32)InValue);
+				return EPropertyBagResult::Success;
+			}
+		case EPropertyBagPropertyType::Int64:
+			{
+				const FInt64Property* Property = CastFieldChecked<FInt64Property>(Desc.CachedProperty);
+				Property->SetPropertyValue(Address, InValue);
+				return EPropertyBagResult::Success;
+			}
+		case EPropertyBagPropertyType::Float:
+			{
+				const FFloatProperty* Property = CastFieldChecked<FFloatProperty>(Desc.CachedProperty);
+				Property->SetPropertyValue(Address, (float)InValue);
+				return EPropertyBagResult::Success;
+			}
+		case EPropertyBagPropertyType::Double:
+			{
+				const FDoubleProperty* Property = CastFieldChecked<FDoubleProperty>(Desc.CachedProperty);
+				Property->SetPropertyValue(Address, (double)InValue);
+				return EPropertyBagResult::Success;
+			}
+		case EPropertyBagPropertyType::Enum:
+			{
+				const FEnumProperty* EnumProperty = CastFieldChecked<FEnumProperty>(Desc.CachedProperty);
+				const FNumericProperty* UnderlyingProperty = EnumProperty->GetUnderlyingProperty();
+				check(UnderlyingProperty);
+				UnderlyingProperty->SetIntPropertyValue(Address, (uint64)InValue);
+				return EPropertyBagResult::Success;
+			}
+		default:
+			return EPropertyBagResult::TypeMismatch;
+		}
+	}
+
+	void CopyMatchingValuesByID(const FConstStructView Source, const FStructView Target)
+	{
+		if (!Source.IsValid() || !Target.IsValid())
+		{
+			return;
+		}
+
+		const UPropertyBag* SourceBagStruct = Cast<const UPropertyBag>(Source.GetScriptStruct());
+		const UPropertyBag* TargetBagStruct = Cast<const UPropertyBag>(Target.GetScriptStruct());
+
+		if (!SourceBagStruct || !TargetBagStruct)
+		{
+			return;
+		}
+
+		// Iterate over source and copy to target if possible. Source is expected to usually have less items.
+		for (const FPropertyBagPropertyDesc& SourceDesc : SourceBagStruct->GetPropertyDescs())
+		{
+			const FPropertyBagPropertyDesc* PotentialTargetDesc = TargetBagStruct->FindPropertyDescByID(SourceDesc.ID);
+			if (PotentialTargetDesc == nullptr
+				|| PotentialTargetDesc->CachedProperty == nullptr
+				|| SourceDesc.CachedProperty == nullptr)
+			{
+				continue;
+			}
+
+			const FPropertyBagPropertyDesc& TargetDesc = *PotentialTargetDesc;
+			void* TargetAddress = Target.GetMutableMemory() + TargetDesc.CachedProperty->GetOffset_ForInternal();
+			const void* SourceAddress = Source.GetMemory() + SourceDesc.CachedProperty->GetOffset_ForInternal();
+			
+			if (TargetDesc.CompatibleType(SourceDesc))
+			{
+				TargetDesc.CachedProperty->CopyCompleteValue(TargetAddress, SourceAddress);
+			}
+			else if (TargetDesc.IsNumericType() && SourceDesc.IsNumericType())
+			{
+				// Try to convert numeric types.
+				if (TargetDesc.IsNumericFloatType())
+				{
+					double Value = 0;
+					if (GetPropertyAsDouble(SourceDesc, SourceAddress, Value) == EPropertyBagResult::Success)
+					{
+						SetPropertyFromDouble(TargetDesc, TargetAddress, Value);
+					}
+				}
+				else
+				{
+					int64 Value = 0;
+					if (GetPropertyAsInt64(SourceDesc, SourceAddress, Value) == EPropertyBagResult::Success)
+					{
+						SetPropertyFromInt64(TargetDesc, TargetAddress, Value);
+					}
+				}
+			}
+			else if ((TargetDesc.IsObjectType() && SourceDesc.IsObjectType())
+				|| (TargetDesc.IsClassType() && SourceDesc.IsClassType()))
+			{
+				// Try convert between compatible objects and classes.
+				const UClass* TargetObjectClass = Cast<const UClass>(TargetDesc.ValueTypeObject);
+				const UClass* SourceObjectClass = Cast<const UClass>(SourceDesc.ValueTypeObject);
+				if (CanCastTo(SourceObjectClass, TargetObjectClass))
+				{
+					const FObjectPropertyBase* TargetProp = CastFieldChecked<FObjectPropertyBase>(TargetDesc.CachedProperty);
+					const FObjectPropertyBase* SourceProp = CastFieldChecked<FObjectPropertyBase>(SourceDesc.CachedProperty);
+					TargetProp->SetObjectPropertyValue(TargetAddress, SourceProp->GetObjectPropertyValue(SourceAddress));
+				}
+			}
+		}
+	}
+
+	// Helper templates to reduce repeated work when dealing with property access.
+
+	template<typename T>
+	TValueOrError<T, EPropertyBagResult> GetValueInt64(const FInstancedPropertyBag& Bag, const FName Name)
+	{
+		const FPropertyBagPropertyDesc* Desc = Bag.FindPropertyDescByName(Name);
+		if (Desc == nullptr || Desc->CachedProperty == nullptr)
+		{
+			return MakeError(EPropertyBagResult::PropertyNotFound);
+		}
+	 
+		check(Bag.GetValue().IsValid());
+		const void* Address = Bag.GetValue().GetMemory() + Desc->CachedProperty->GetOffset_ForInternal();
+		int64 IntValue = 0;
+		const EPropertyBagResult Result = GetPropertyAsInt64(*Desc, Address, IntValue);
+		if (Result != EPropertyBagResult::Success)
+		{
+			return MakeError(Result);
+		}
+		return MakeValue(IntValue);
+	}
+
+	template<typename T>
+	TValueOrError<T, EPropertyBagResult> GetValueDouble(const FInstancedPropertyBag& Bag, const FName Name)
+	{
+		const FPropertyBagPropertyDesc* Desc = Bag.FindPropertyDescByName(Name);
+		if (Desc == nullptr || Desc->CachedProperty == nullptr)
+		{
+			return MakeError(EPropertyBagResult::PropertyNotFound);
+		}
+
+		check(Bag.GetValue().IsValid());
+		const void* Address = Bag.GetValue().GetMemory() + Desc->CachedProperty->GetOffset_ForInternal();
+		double DblValue = 0;
+		const EPropertyBagResult Result = GetPropertyAsDouble(*Desc, Address, DblValue);
+		if (Result != EPropertyBagResult::Success)
+		{
+			return MakeError(Result);
+		}
+		return MakeValue(DblValue);
+	}
+
+	template<typename T, typename PropT>
+	TValueOrError<T, EPropertyBagResult> GetValue(const FInstancedPropertyBag& Bag, const FName Name)
+	{
+		const FPropertyBagPropertyDesc* Desc = Bag.FindPropertyDescByName(Name);
+		if (Desc == nullptr || Desc->CachedProperty == nullptr || !Desc->CachedProperty->IsA<PropT>())
+		{
+			return MakeError(EPropertyBagResult::PropertyNotFound);
+		}
+		check(Desc->CachedProperty);
+		const PropT* Property = CastFieldChecked<PropT>(Desc->CachedProperty);
+		check(Bag.GetValue().IsValid());
+		const void* Address = Bag.GetValue().GetMemory() + Desc->CachedProperty->GetOffset_ForInternal();
+		return MakeValue(Property->GetPropertyValue(Address));
+	}
+
+
+	template<typename T>
+	EPropertyBagResult SetValueInt64(const FInstancedPropertyBag& Bag, const FName Name, const T Value)
+	{
+		const FPropertyBagPropertyDesc* Desc = Bag.FindPropertyDescByName(Name);
+		if (Desc == nullptr || Desc->CachedProperty == nullptr)
+		{
+			return EPropertyBagResult::PropertyNotFound;
+		}
+		check(Desc->CachedProperty);
+		check(Bag.GetValue().IsValid());
+		void* Address = Bag.GetMutableValue().GetMutableMemory() + Desc->CachedProperty->GetOffset_ForInternal();
+		return SetPropertyFromInt64(*Desc, Address, (int64)Value);
+	}
+
+	template<typename T>
+	EPropertyBagResult SetValueDouble(const FInstancedPropertyBag& Bag, const FName Name, const T Value)
+	{
+		const FPropertyBagPropertyDesc* Desc = Bag.FindPropertyDescByName(Name);
+		if (Desc == nullptr || Desc->CachedProperty == nullptr)
+		{
+			return EPropertyBagResult::PropertyNotFound;
+		}
+		check(Desc->CachedProperty);
+		check(Bag.GetValue().IsValid());
+		void* Address = Bag.GetMutableValue().GetMutableMemory() + Desc->CachedProperty->GetOffset_ForInternal();
+		return SetPropertyFromDouble(*Desc, Address, (double)Value);
+	}
+
+	template<typename T, typename PropT>
+	EPropertyBagResult SetValue(const FInstancedPropertyBag& Bag, const FName Name, const T& Value)
+	{
+		const FPropertyBagPropertyDesc* Desc = Bag.FindPropertyDescByName(Name);
+		if (Desc == nullptr || Desc->CachedProperty == nullptr || !Desc->CachedProperty->IsA<PropT>())
+		{
+			return EPropertyBagResult::PropertyNotFound;
+		}
+		check(Desc->CachedProperty);
+		const PropT* Property = CastFieldChecked<PropT>(Desc->CachedProperty);
+		check(Bag.GetValue().IsValid());
+		void* Address = Bag.GetMutableValue().GetMutableMemory() + Desc->CachedProperty->GetOffset_ForInternal();
+		Property->SetPropertyValue(Address, Value);
+		return EPropertyBagResult::Success;
+	}
+
+	void RemovePropertyByName(TArray<FPropertyBagPropertyDesc>& Descs, const FName PropertyName, const int32 StartIndex = 0)
+	{
+		// Remove properties which dont have unique name.
+		for (int32 Index = StartIndex; Index < Descs.Num(); Index++)
+		{
+			if (Descs[Index].Name == PropertyName)
+			{
+				Descs.RemoveAt(Index);
+				Index--;
+			}
+		}
+	}
+
+
+}; // UE::StructUtils::Private
+
+
+//----------------------------------------------------------------//
+//  FPropertyBagPropertyDesc
+//----------------------------------------------------------------//
+
+static FArchive& operator<<(FArchive& Ar, FPropertyBagPropertyDesc& Bag)
+{
+	Ar << Bag.ValueTypeObject;
+	Ar << Bag.ID;
+	Ar << Bag.Name;
+	Ar << Bag.ValueType;
+	return Ar;
+}
+
+bool FPropertyBagPropertyDesc::IsNumericType() const
+{
+	switch (ValueType)
+	{
+	case EPropertyBagPropertyType::Bool: return true;
+	case EPropertyBagPropertyType::Byte: return true;
+	case EPropertyBagPropertyType::Int32: return true;
+	case EPropertyBagPropertyType::Int64: return true;
+	case EPropertyBagPropertyType::Float: return true;
+	case EPropertyBagPropertyType::Double: return true;
+	default: return false;
+	}
+}
+
+bool FPropertyBagPropertyDesc::IsNumericFloatType() const
+{
+	switch (ValueType)
+	{
+	case EPropertyBagPropertyType::Float: return true;
+	case EPropertyBagPropertyType::Double: return true;
+	default: return false;
+	}
+}
+
+bool FPropertyBagPropertyDesc::IsObjectType() const
+{
+	switch (ValueType)
+	{
+	case EPropertyBagPropertyType::Object: return true;
+	case EPropertyBagPropertyType::SoftObject: return true;
+	default: return false;
+	}
+}
+
+bool FPropertyBagPropertyDesc::IsClassType() const
+{
+	switch (ValueType)
+	{
+	case EPropertyBagPropertyType::Class: return true;
+	case EPropertyBagPropertyType::SoftClass: return true;
+	default: return false;
+	}
+}
+
+bool FPropertyBagPropertyDesc::CompatibleType(const FPropertyBagPropertyDesc& Other) const
+{
+	// Values must match.
+	if (ValueType != Other.ValueType)
+	{
+		return false;
+	}
+
+	// Struct and enum must have same value type class
+	if (ValueType == EPropertyBagPropertyType::Enum || ValueType == EPropertyBagPropertyType::Struct)
+	{
+		return ValueTypeObject == Other.ValueTypeObject; 
+	}
+
+	// Objects should be castable.
+	if (ValueType == EPropertyBagPropertyType::Object)
+	{
+		const UClass* ObjectClass = Cast<const UClass>(ValueTypeObject);
+		const UClass* OtherObjectClass = Cast<const UClass>(Other.ValueTypeObject);
+		return UE::StructUtils::Private::CanCastTo(OtherObjectClass, ObjectClass);
+	}
+
+	return true;
+}
+
+//----------------------------------------------------------------//
+//  FInstancedPropertyBag
+//----------------------------------------------------------------//
+
+void FInstancedPropertyBag::InitializeFromBagStruct(const UPropertyBag* NewBagStruct)
+{
+	Value.InitializeAs(NewBagStruct);
+}
+
+void FInstancedPropertyBag::CopyMatchingValuesByID(const FInstancedPropertyBag& Other) const
+{
+	UE::StructUtils::Private::CopyMatchingValuesByID(Other.Value, Value);
+}
+
+void FInstancedPropertyBag::AddProperties(const TConstArrayView<FPropertyBagPropertyDesc> NewDescs)
+{
+	TArray<FPropertyBagPropertyDesc> Descs;
+	if (const UPropertyBag* CurrentBagStruct = GetPropertyBagStruct())
+	{
+		Descs = CurrentBagStruct->GetPropertyDescs();
+	}
+
+	for (const FPropertyBagPropertyDesc& NewDesc : NewDescs)
+	{
+		FPropertyBagPropertyDesc* ExistingProperty = Descs.FindByPredicate([&NewDesc](const FPropertyBagPropertyDesc& Desc) { return Desc.Name == NewDesc.Name; });
+		if (ExistingProperty != nullptr)
+		{
+			ExistingProperty->ValueType = NewDesc.ValueType;
+			ExistingProperty->ValueTypeObject = NewDesc.ValueTypeObject;
+		}
+		else
+		{
+			Descs.Add(NewDesc);
+		}
+	}
+
+	const UPropertyBag* NewBagStruct = UPropertyBag::GetOrCreateFromDescs(Descs);
+	MigrateToNewBagStruct(NewBagStruct);
+}
+	
+void FInstancedPropertyBag::AddProperty(const FName InName, const EPropertyBagPropertyType InValueType, UObject* InValueTypeObject)
+{
+	AddProperties({ FPropertyBagPropertyDesc(InName, InValueType, InValueTypeObject) });
+}
+
+void FInstancedPropertyBag::RemovePropertiesByName(const TConstArrayView<FName> PropertiesToRemove)
+{
+	TArray<FPropertyBagPropertyDesc> Descs;
+	if (const UPropertyBag* CurrentBagStruct = GetPropertyBagStruct())
+	{
+		Descs = CurrentBagStruct->GetPropertyDescs();
+	}
+
+	for (const FName Name : PropertiesToRemove)
+	{
+		UE::StructUtils::Private::RemovePropertyByName(Descs, Name);
+	}
+
+	const UPropertyBag* NewBagStruct = UPropertyBag::GetOrCreateFromDescs(Descs);
+	MigrateToNewBagStruct(NewBagStruct);
+}
+	
+void FInstancedPropertyBag::RemovePropertyByName(const FName PropertyToRemove)
+{
+	RemovePropertiesByName({ PropertyToRemove });
+}
+
+void FInstancedPropertyBag::MigrateToNewBagStruct(const UPropertyBag* NewBagStruct)
+{
+	FInstancedStruct NewValue(NewBagStruct);
+
+	UE::StructUtils::Private::CopyMatchingValuesByID(Value, NewValue);
+	
+	Value = MoveTemp(NewValue);
+}
+
+const UPropertyBag* FInstancedPropertyBag::GetPropertyBagStruct() const
+{
+	return Value.IsValid() ? Cast<const UPropertyBag>(Value.GetScriptStruct()) : nullptr;
+}
+
+const FPropertyBagPropertyDesc* FInstancedPropertyBag::FindPropertyDescByID(const FGuid ID) const
+{
+	if (const UPropertyBag* BagStruct = GetPropertyBagStruct())
+	{
+		return BagStruct->FindPropertyDescByID(ID);
+	}
+	return nullptr;
+}
+	
+const FPropertyBagPropertyDesc* FInstancedPropertyBag::FindPropertyDescByName(const FName Name) const
+{
+	if (const UPropertyBag* BagStruct = GetPropertyBagStruct())
+	{
+		return BagStruct->FindPropertyDescByName(Name);
+	}
+	return nullptr;
+}
+
+TValueOrError<bool, EPropertyBagResult> FInstancedPropertyBag::GetValueBool(const FName Name) const
+{
+	TValueOrError<int64, EPropertyBagResult> Result = UE::StructUtils::Private::GetValueInt64<int64>(*this, Name);
+	if (!Result.IsValid())
+	{
+		return MakeError(Result.GetError());
+	}
+	return MakeValue(Result.GetValue() != 0);
+}
+
+TValueOrError<uint8, EPropertyBagResult> FInstancedPropertyBag::GetValueByte(const FName Name) const
+{
+	return UE::StructUtils::Private::GetValueInt64<uint8>(*this, Name);
+}
+
+TValueOrError<int32, EPropertyBagResult> FInstancedPropertyBag::GetValueInt32(const FName Name) const
+{
+	return UE::StructUtils::Private::GetValueInt64<int32>(*this, Name);
+}
+
+TValueOrError<int64, EPropertyBagResult> FInstancedPropertyBag::GetValueInt64(const FName Name) const
+{
+	return UE::StructUtils::Private::GetValueInt64<int64>(*this, Name);
+}
+
+TValueOrError<float, EPropertyBagResult> FInstancedPropertyBag::GetValueFloat(const FName Name) const
+{
+	return UE::StructUtils::Private::GetValueDouble<float>(*this, Name);
+}
+
+TValueOrError<double, EPropertyBagResult> FInstancedPropertyBag::GetValueDouble(const FName Name) const
+{
+	return UE::StructUtils::Private::GetValueDouble<double>(*this, Name);
+}
+
+TValueOrError<FName, EPropertyBagResult> FInstancedPropertyBag::GetValueName(const FName Name) const
+{
+	return UE::StructUtils::Private::GetValue<FName, FNameProperty>(*this, Name);
+}
+
+TValueOrError<FString, EPropertyBagResult> FInstancedPropertyBag::GetValueString(const FName Name) const
+{
+	return UE::StructUtils::Private::GetValue<FString, FStrProperty>(*this, Name);
+}
+
+TValueOrError<FText, EPropertyBagResult> FInstancedPropertyBag::GetValueText(const FName Name) const
+{
+	return UE::StructUtils::Private::GetValue<FText, FTextProperty>(*this, Name);
+}
+
+TValueOrError<uint8, EPropertyBagResult> FInstancedPropertyBag::GetValueEnum(const FName Name, const UEnum* RequestedEnum) const
+{
+	const FPropertyBagPropertyDesc* Desc = FindPropertyDescByName(Name);
+	if (Desc == nullptr || Desc->ValueType != EPropertyBagPropertyType::Enum)
+	{
+		return MakeError(Desc == nullptr ? EPropertyBagResult::PropertyNotFound : EPropertyBagResult::TypeMismatch);
+	}
+	check(Desc->CachedProperty);
+	const FEnumProperty* EnumProperty = CastFieldChecked<FEnumProperty>(Desc->CachedProperty);
+	const FNumericProperty* UnderlyingProperty = EnumProperty->GetUnderlyingProperty();
+	check(UnderlyingProperty);
+	
+	if (RequestedEnum != EnumProperty->GetEnum())
+	{
+		return MakeError(EPropertyBagResult::TypeMismatch);
+	}
+	
+	check(Value.IsValid());
+	const void* Address = Value.GetMutableMemory() + Desc->CachedProperty->GetOffset_ForInternal();
+	return MakeValue((uint8)UnderlyingProperty->GetUnsignedIntPropertyValue(Address));
+}
+
+TValueOrError<FStructView, EPropertyBagResult> FInstancedPropertyBag::GetValueStruct(const FName Name, const UScriptStruct* RequestedStruct) const
+{
+	const FPropertyBagPropertyDesc* Desc = FindPropertyDescByName(Name);
+	if (Desc == nullptr || Desc->ValueType != EPropertyBagPropertyType::Struct)
+	{
+		return MakeError(Desc == nullptr ? EPropertyBagResult::PropertyNotFound : EPropertyBagResult::TypeMismatch);
+	}
+	check(Desc->CachedProperty);
+	const FStructProperty* StructProperty = CastFieldChecked<FStructProperty>(Desc->CachedProperty);
+	check(StructProperty->Struct);
+
+	if (RequestedStruct != nullptr && UE::StructUtils::Private::CanCastTo(StructProperty->Struct, RequestedStruct) == false)
+	{
+		return MakeError(EPropertyBagResult::TypeMismatch);
+	}
+	
+	check(Value.IsValid());
+	void* Address = Value.GetMutableMemory() + Desc->CachedProperty->GetOffset_ForInternal();
+	return MakeValue(FStructView(StructProperty->Struct, (uint8*)Address));
+}
+
+TValueOrError<UObject*, EPropertyBagResult> FInstancedPropertyBag::GetValueObject(const FName Name, const UClass* RequestedClass) const
+{
+	const FPropertyBagPropertyDesc* Desc = FindPropertyDescByName(Name);
+	if (Desc == nullptr || (Desc->ValueType != EPropertyBagPropertyType::Object && Desc->ValueType != EPropertyBagPropertyType::SoftObject))
+	{
+		return MakeError(Desc == nullptr ? EPropertyBagResult::PropertyNotFound : EPropertyBagResult::TypeMismatch);
+	}
+	check(Desc->CachedProperty);
+	const FObjectPropertyBase* ObjectProperty = CastFieldChecked<FObjectPropertyBase>(Desc->CachedProperty);
+	check(ObjectProperty->PropertyClass);
+
+	if (RequestedClass != nullptr && UE::StructUtils::Private::CanCastTo(ObjectProperty->PropertyClass, RequestedClass) == false)
+	{
+		return MakeError(EPropertyBagResult::TypeMismatch);
+	}
+	
+	check(Value.IsValid());
+	const void* Address = Value.GetMemory() + Desc->CachedProperty->GetOffset_ForInternal();
+	return MakeValue(ObjectProperty->GetObjectPropertyValue(Address));
+}
+
+TValueOrError<UClass*, EPropertyBagResult> FInstancedPropertyBag::GetValueClass(const FName Name) const
+{
+	const FPropertyBagPropertyDesc* Desc = FindPropertyDescByName(Name);
+	if (Desc == nullptr || (Desc->ValueType != EPropertyBagPropertyType::Class && Desc->ValueType != EPropertyBagPropertyType::SoftClass))
+	{
+		return MakeError(Desc == nullptr ? EPropertyBagResult::PropertyNotFound : EPropertyBagResult::TypeMismatch);
+	}
+	check(Desc->CachedProperty);
+	const FObjectPropertyBase* ObjectProperty = CastFieldChecked<FObjectPropertyBase>(Desc->CachedProperty);
+	check(Value.IsValid());
+	const void* Address = Value.GetMemory() + Desc->CachedProperty->GetOffset_ForInternal();
+	return MakeValue(CastChecked<UClass>(ObjectProperty->GetObjectPropertyValue(Address)));
+}
+
+
+EPropertyBagResult FInstancedPropertyBag::SetValueBool(const FName Name, const bool bInValue) const
+{
+	return UE::StructUtils::Private::SetValueInt64<int64>(*this, Name, bInValue ? 1 : 0);
+}
+
+EPropertyBagResult FInstancedPropertyBag::SetValueByte(const FName Name, const uint8 InValue) const
+{
+	return UE::StructUtils::Private::SetValueInt64<uint8>(*this, Name, InValue);
+}
+
+EPropertyBagResult FInstancedPropertyBag::SetValueInt32(const FName Name, const int32 InValue) const
+{
+	return UE::StructUtils::Private::SetValueInt64<int32>(*this, Name, InValue);
+}
+
+EPropertyBagResult FInstancedPropertyBag::SetValueInt64(const FName Name, const int64 InValue) const
+{
+	return UE::StructUtils::Private::SetValueInt64<int64>(*this, Name, InValue);
+}
+
+EPropertyBagResult FInstancedPropertyBag::SetValueFloat(const FName Name, const float InValue) const
+{
+	return UE::StructUtils::Private::SetValueDouble<float>(*this, Name, InValue);
+}
+
+EPropertyBagResult FInstancedPropertyBag::SetValueDouble(const FName Name, const double InValue) const
+{
+	return UE::StructUtils::Private::SetValueDouble<double>(*this, Name, InValue);
+}
+
+EPropertyBagResult FInstancedPropertyBag::SetValueName(const FName Name, const FName InValue) const
+{
+	return UE::StructUtils::Private::SetValue<FName, FNameProperty>(*this, Name, InValue);
+}
+
+EPropertyBagResult FInstancedPropertyBag::SetValueString(const FName Name, const FString& InValue) const
+{
+	return UE::StructUtils::Private::SetValue<FString, FStrProperty>(*this, Name, InValue);
+}
+
+EPropertyBagResult FInstancedPropertyBag::SetValueText(const FName Name, const FText& InValue) const
+{
+	return UE::StructUtils::Private::SetValue<FText, FTextProperty>(*this, Name, InValue);
+}
+
+EPropertyBagResult FInstancedPropertyBag::SetValueEnum(const FName Name, const uint8 InValue, const UEnum* Enum) const
+{
+	const FPropertyBagPropertyDesc* Desc = FindPropertyDescByName(Name);
+	if (Desc == nullptr || Desc->ValueType != EPropertyBagPropertyType::Enum)
+	{
+		return Desc == nullptr ? EPropertyBagResult::PropertyNotFound : EPropertyBagResult::TypeMismatch;
+	}
+	check(Desc->CachedProperty);
+	const FEnumProperty* EnumProperty = CastFieldChecked<FEnumProperty>(Desc->CachedProperty);
+	const FNumericProperty* UnderlyingProperty = EnumProperty->GetUnderlyingProperty();
+	check(UnderlyingProperty);
+
+	if (Enum != EnumProperty->GetEnum())
+	{
+		return EPropertyBagResult::TypeMismatch;
+	}
+	
+	check(Value.IsValid());
+	void* Address = Value.GetMutableMemory() + Desc->CachedProperty->GetOffset_ForInternal();
+	UnderlyingProperty->SetIntPropertyValue(Address, (uint64)InValue);
+	
+	return EPropertyBagResult::Success;
+}
+
+EPropertyBagResult FInstancedPropertyBag::SetValueStruct(const FName Name, FConstStructView InValue) const
+{
+	const FPropertyBagPropertyDesc* Desc = FindPropertyDescByName(Name);
+	if (Desc == nullptr || Desc->ValueType != EPropertyBagPropertyType::Struct)
+	{
+		return Desc == nullptr ? EPropertyBagResult::PropertyNotFound : EPropertyBagResult::TypeMismatch;
+	}
+	check(Desc->CachedProperty);
+	const FStructProperty* StructProperty = CastFieldChecked<FStructProperty>(Desc->CachedProperty);
+	check(StructProperty->Struct);
+
+	if (InValue.GetScriptStruct() && InValue.GetScriptStruct() != StructProperty->Struct)
+	{
+		return EPropertyBagResult::TypeMismatch;
+	}
+
+	check(Value.IsValid());
+	void* Address = Value.GetMutableMemory() + Desc->CachedProperty->GetOffset_ForInternal();
+
+	if (InValue.IsValid())
+	{
+		StructProperty->Struct->CopyScriptStruct(Address, InValue.GetMemory());
+	}
+	else
+	{
+		StructProperty->Struct->ClearScriptStruct(Address);
+	}
+
+	return EPropertyBagResult::Success;
+}
+
+EPropertyBagResult FInstancedPropertyBag::SetValueObject(const FName Name, UObject* InValue) const
+{
+	const FPropertyBagPropertyDesc* Desc = FindPropertyDescByName(Name);
+	if (Desc == nullptr || (Desc->ValueType != EPropertyBagPropertyType::Object && Desc->ValueType != EPropertyBagPropertyType::SoftObject))
+	{
+		return Desc == nullptr ? EPropertyBagResult::PropertyNotFound : EPropertyBagResult::TypeMismatch;
+	}
+	check(Desc->CachedProperty);
+	const FObjectPropertyBase* ObjectProperty = CastFieldChecked<FObjectPropertyBase>(Desc->CachedProperty);
+	check(ObjectProperty->PropertyClass);
+
+	if (InValue && UE::StructUtils::Private::CanCastTo(InValue->GetClass(), ObjectProperty->PropertyClass) == false)
+	{
+		return EPropertyBagResult::TypeMismatch;
+	}
+	
+	check(Value.IsValid());
+	void* Address = Value.GetMutableMemory() + Desc->CachedProperty->GetOffset_ForInternal();
+	ObjectProperty->SetObjectPropertyValue(Address, InValue);
+
+	return EPropertyBagResult::Success;
+}
+
+EPropertyBagResult FInstancedPropertyBag::SetValueClass(const FName Name, UClass* InValue) const
+{
+	const FPropertyBagPropertyDesc* Desc = FindPropertyDescByName(Name);
+	if (Desc == nullptr || (Desc->ValueType != EPropertyBagPropertyType::Class && Desc->ValueType != EPropertyBagPropertyType::SoftClass))
+	{
+		return Desc == nullptr ? EPropertyBagResult::PropertyNotFound : EPropertyBagResult::TypeMismatch;
+	}
+	check(Desc->CachedProperty);
+	
+	check(Value.IsValid());
+	void* Address = Value.GetMutableMemory() + Desc->CachedProperty->GetOffset_ForInternal();
+
+	if (Desc->ValueType == EPropertyBagPropertyType::Class)
+	{
+		const FClassProperty* ClassProperty = CastFieldChecked<FClassProperty>(Desc->CachedProperty);
+		if (InValue && InValue->IsChildOf(ClassProperty->MetaClass) == false)
+		{
+			return EPropertyBagResult::TypeMismatch;
+		}
+		ClassProperty->SetObjectPropertyValue(Address, InValue);
+	}
+	else
+	{
+		const FSoftClassProperty* ClassProperty = CastFieldChecked<FSoftClassProperty>(Desc->CachedProperty);
+		if (InValue && InValue->IsChildOf(ClassProperty->MetaClass) == false)
+		{
+			return EPropertyBagResult::TypeMismatch;
+		}
+		ClassProperty->SetObjectPropertyValue(Address, InValue);
+	}
+
+	return EPropertyBagResult::Success;
+}
+
+bool FInstancedPropertyBag::Serialize(FArchive& Ar)
+{
+	enum class EVersion : uint8
+	{
+		InitialVersion = 0,
+		// -----<new versions can be added above this line>-----
+		VersionPlusOne,
+		LatestVersion = VersionPlusOne - 1
+	};
+
+	EVersion Version = EVersion::LatestVersion;
+
+	Ar << Version;
+
+	UPropertyBag* BagStruct = Cast<UPropertyBag>(const_cast<UScriptStruct*>(Value.GetScriptStruct()));
+	bool bHasData = (BagStruct != nullptr);
+	
+	Ar << bHasData;
+	
+	if (bHasData)
+	{
+		// The script struct class is not serialized, the properties are serialized and type is created based on that.
+		if (Ar.IsLoading())
+		{
+			TArray<FPropertyBagPropertyDesc> PropertyDescs;
+
+			Ar << PropertyDescs;
+
+			BagStruct = const_cast<UPropertyBag*>(UPropertyBag::GetOrCreateFromDescs(PropertyDescs));
+			Value.InitializeAs(BagStruct);
+		}
+		if (Ar.IsSaving())
+		{
+			Ar << BagStruct->PropertyDescs;
+		}
+
+		if (ensureMsgf(Value.GetMutableMemory() != nullptr, TEXT("A valid script struct should always have allocated memory")))
+		{
+			BagStruct->SerializeItem(Ar, Value.GetMutableMemory(), /*Defaults*/nullptr);
+		}
+	}
+
+	return true;
+}
+
+//----------------------------------------------------------------//
+//  UPropertyBag
+//----------------------------------------------------------------//
+
+const UPropertyBag* UPropertyBag::GetOrCreateFromDescs(const TConstArrayView<FPropertyBagPropertyDesc> PropertyDescs)
+{
+	const uint64 BagHash = UE::StructUtils::Private::CalcPropertyDescArrayHash(PropertyDescs);
+	const FString ScriptStructName = FString::Printf(TEXT("PropertyBag_%llx"), BagHash);
+
+	if (const UPropertyBag* ExistingBag = FindObject<UPropertyBag>(GetTransientPackage(), *ScriptStructName))
+	{
+		return ExistingBag;
+	}
+
+	// Validate descriptors
+	bool bAllValid = true;
+	for (const FPropertyBagPropertyDesc& Desc : PropertyDescs)
+	{
+		if (Desc.ValueType == EPropertyBagPropertyType::Struct)
+		{
+			if (Desc.ValueTypeObject == nullptr || Desc.ValueTypeObject->GetClass()->IsChildOf(UScriptStruct::StaticClass()) == false)
+			{
+				ensureMsgf(false, TEXT("Struct property must have valid struct type object."));
+				bAllValid = false;
+			}
+		}
+		else if (Desc.ValueType == EPropertyBagPropertyType::Enum)
+		{
+			if (Desc.ValueTypeObject == nullptr || Desc.ValueTypeObject->GetClass()->IsChildOf(UEnum::StaticClass()) == false)
+			{
+				ensureMsgf(false, TEXT("Enum property must have valid enum type object."));
+				bAllValid = false;
+			}
+		}
+		else if (Desc.ValueType == EPropertyBagPropertyType::Object || Desc.ValueType == EPropertyBagPropertyType::SoftObject)
+		{
+			if (Desc.ValueTypeObject == nullptr)
+			{
+				ensureMsgf(false, TEXT("Object property must have valid object type object."));
+				bAllValid = false;
+			}
+		}
+		else if (Desc.ValueType == EPropertyBagPropertyType::Class || Desc.ValueType == EPropertyBagPropertyType::SoftClass)
+		{
+			if (Desc.ValueTypeObject == nullptr || Desc.ValueTypeObject->GetClass()->IsChildOf(UClass::StaticClass()) == false)
+			{
+				ensureMsgf(false, TEXT("Class property must have valid class type object."));
+				bAllValid = false;
+			}
+		}
+	}
+	if (!bAllValid)
+	{
+		return nullptr;
+	}
+	
+	UPropertyBag* NewBag = NewObject<UPropertyBag>(GetTransientPackage(), *ScriptStructName, RF_Standalone | RF_Transient);
+
+	NewBag->PropertyDescs = PropertyDescs;
+
+	// Remove properties with same name
+	for (int32 Index = 0; Index < NewBag->PropertyDescs.Num() - 1; Index++)
+	{
+		UE::StructUtils::Private::RemovePropertyByName(NewBag->PropertyDescs, NewBag->PropertyDescs[Index].Name, Index + 1);
+	}
+	
+	// Add properties (AddCppProperty() adds them backwards in the linked list)
+	for (int32 DescIndex = NewBag->PropertyDescs.Num() - 1; DescIndex >= 0; DescIndex--)
+	{
+		FPropertyBagPropertyDesc& Desc = NewBag->PropertyDescs[DescIndex];
+
+		if (!Desc.ID.IsValid())
+		{
+			Desc.ID = FGuid::NewGuid();
+		}
+		
+		if (FProperty* NewProperty = UE::StructUtils::Private::CreatePropertyFromDesc(Desc, NewBag))
+		{
+			NewProperty->SetPropertyFlags(CPF_Edit);
+			NewBag->AddCppProperty(NewProperty);
+			Desc.CachedProperty = NewProperty; 
+		}
+	}
+
+	NewBag->Bind();
+	NewBag->StaticLink(/*RelinkExistingProperties*/true);
+
+	return NewBag;
+}
+
+const FPropertyBagPropertyDesc* UPropertyBag::FindPropertyDescByID(const FGuid ID) const
+{
+	return PropertyDescs.FindByPredicate([&ID](const FPropertyBagPropertyDesc& Desc) { return Desc.ID == ID; });
+}
+
+const FPropertyBagPropertyDesc* UPropertyBag::FindPropertyDescByName(const FName Name) const
+{
+	return PropertyDescs.FindByPredicate([&Name](const FPropertyBagPropertyDesc& Desc) { return Desc.Name == Name; });
+}
diff --git a/Engine/Plugins/Experimental/StructUtils/Source/StructUtils/Public/PropertyBag.h b/Engine/Plugins/Experimental/StructUtils/Source/StructUtils/Public/PropertyBag.h
new file mode 100644
--- /dev/null
+++ b/Engine/Plugins/Experimental/StructUtils/Source/StructUtils/Public/PropertyBag.h
@@ -0,0 +1,366 @@
+﻿// Copyright Epic Games, Inc. All Rights Reserved.
+
+#pragma once
+
+#include "InstancedStruct.h"
+#include "StructView.h"
+#include "Templates/ValueOrError.h"
+#include "PropertyBag.generated.h"
+
+/** Property bag property type, loosely based on BluePrint pin types. */
+UENUM()
+enum class EPropertyBagPropertyType : uint8
+{
+	None UMETA(Hidden),
+	Bool,
+	Byte,
+	Int32,
+	Int64,
+	Float,
+	Double,
+	Name,
+	String,
+	Text,
+	Enum UMETA(Hidden),
+	Struct UMETA(Hidden),
+	Object UMETA(Hidden),
+	SoftObject UMETA(Hidden),
+	Class UMETA(Hidden),
+	SoftClass UMETA(Hidden),
+};
+
+/** Getter and setter result code. */
+UENUM()
+enum class EPropertyBagResult : uint8
+{
+	Success,			// Operation succeeded.
+	TypeMismatch,		// Tried to access mismatching type (e.g. setting a struct to bool)
+	PropertyNotFound,	// Could not find property of specified name.
+};
+
+/** Describes a property in the property bag. */
+USTRUCT()
+struct STRUCTUTILS_API FPropertyBagPropertyDesc
+{
+	GENERATED_BODY()
+
+	FPropertyBagPropertyDesc() = default;
+	FPropertyBagPropertyDesc(const FName InName, const EPropertyBagPropertyType InValueType, UObject* InValueTypeObject = nullptr)
+		: ValueTypeObject(InValueTypeObject)
+		, Name(InName)
+		, ValueType(InValueType)
+	{
+	}
+
+	/** @return true if the two descriptors have the same type. Object types are compatible if Other can be cast to this type. */
+	bool CompatibleType(const FPropertyBagPropertyDesc& Other) const;
+
+	/** @return true if the property type is numeric (bool, int32, int64, float, double, enum) */
+	bool IsNumericType() const;
+	
+	/** @return true if the property type is floating point numeric (float, double) */
+	bool IsNumericFloatType() const;
+
+	/** @return true if the property type is object or soft object */
+	bool IsObjectType() const;
+
+	/** @return true if the property type is class or soft class */
+	bool IsClassType() const;
+
+	
+	/** Pointer to object that defines the Enum, Struct, or Class. */
+	UPROPERTY(EditAnywhere, Category="Default")
+	TObjectPtr<const UObject> ValueTypeObject = nullptr;
+
+	/** Unique ID for this property. Used as main identifier when copying values over. */
+	UPROPERTY(EditAnywhere, Category="Default")
+	FGuid ID;
+
+	/** Name for the property. */
+	UPROPERTY(EditAnywhere, Category="Default")
+	FName Name;
+
+	/** Type of the value described by this property. */
+	UPROPERTY(EditAnywhere, Category="Default")
+	EPropertyBagPropertyType ValueType = EPropertyBagPropertyType::None;
+
+	/** Cached property pointer, set in UPropertyBag::GetOrCreateFromDescs. */
+	const FProperty* CachedProperty = nullptr;
+};
+
+/**
+ * Instanced property bag allows to create and store a bag of properties.
+ *
+ * When used as editable property, the UI allows properties to be added and removed, and values to be set.
+ * The value is stored as a struct, the type of the value is never serialized, instead the composition of the properties
+ * is saved with the instance, and the type is recreated on load. The types with same composition of properties share same type (based on hashing).
+ *
+ * NOTE: Adding or removing properties to the instance is quite expensive as it will create new UPropertyBag, reallocate memory, and copy all values over. 
+ *
+ * Example usage, this allows the bag to be configured in the UI:
+ *
+ *		UPROPERTY(EditDefaultsOnly, Category = Common)
+ *		FInstancedPropertyBag Bag;
+ *
+ * Changing the layout from code:
+ *
+ *		static const FName TemperatureName(TEXT("Temperature"));
+ *		static const FName IsHotName(TEXT("bIsHot"));
+ *
+ *		FInstancedPropertyBag Bag;
+ *
+ *		// Add properties to the bag, and set their values.
+ *		// Adding or removing properties is not cheap, so better do it in batches.
+ *		Bag.AddProperties({
+ *			{ TemperatureName, EPropertyBagPropertyType::Float },
+ *			{ CountName, EPropertyBagPropertyType::Int32 }
+ *		});
+ *
+ *		// Amend the bag with a new property.
+ *		Bag.AddProperty(IsHotName, EPropertyBagPropertyType::Bool);
+ *		Bag.SetValueBool(IsHotName, true);
+ *
+ *		// Get value and use the result
+ *		if (auto Temperature = Bag.GetValueFloat(TemperatureName); Temperature.IsValid())
+ *		{
+ *			float Val = Temperature.GetValue();
+ *		}
+ */
+
+class UPropertyBag;
+
+USTRUCT()
+struct STRUCTUTILS_API FInstancedPropertyBag
+{
+	GENERATED_BODY()
+
+	FInstancedPropertyBag() = default;
+	FInstancedPropertyBag(const FInstancedPropertyBag& Other) = default;
+	FInstancedPropertyBag(FInstancedPropertyBag&& Other) = default;
+
+	FInstancedPropertyBag& operator=(const FInstancedPropertyBag& InOther) = default;
+	FInstancedPropertyBag& operator=(FInstancedPropertyBag&& InOther) = default;
+
+	/** Resets the instance to empty. */
+	void Reset()
+	{
+		Value.Reset();
+	}
+
+	/** Initializes the instance from an bag struct. */
+	void InitializeFromBagStruct(const UPropertyBag* NewBagStruct);
+	
+	/**
+	 * Copies matching property values from another bag of potentially mismatching layout.
+	 * The properties are matched between the bags based on the property ID.
+	 * @param Other Reference to the bag to copy the values from
+	 */
+	void CopyMatchingValuesByID(const FInstancedPropertyBag& NewDescs) const;
+
+	/**
+	 * Adds properties to the bag. If property of same name already exists, it will be replaced with the new type.
+	 * Numeric property values will be converted if possible, when a property's type changes.
+	 * @param Descs Descriptors of new properties to add.
+	 */
+	void AddProperties(const TConstArrayView<FPropertyBagPropertyDesc> Descs);
+	
+	/**
+	 * Adds a new property to the bag. If property of same name already exists, it will be replaced with the new type.
+	 * Numeric property values will be converted if possible, when a property's type changes.
+	 * @param Descs Descriptors of new properties to add.
+	 */
+	void AddProperty(const FName InName, const EPropertyBagPropertyType InValueType, UObject* InValueTypeObject = nullptr);
+
+	/**
+	 * Removes properties from the bag by name if they exists.
+	 */
+	void RemovePropertiesByName(const TConstArrayView<FName> PropertiesToRemove);
+	
+	/**
+	 * Removes a property from the bag by name if it exists.
+	 */
+	void RemovePropertyByName(const FName PropertyToRemove);
+	
+	/**
+	 * Changes the type of this bag and migrates existing values.
+	 * The properties are matched between the bags based on the property ID.
+	 * @param NewBagStruct Pointer to the new type.
+	 */
+	void MigrateToNewBagStruct(const UPropertyBag* NewBagStruct);
+
+	/** @return pointer to the property bag struct. */ 
+	const UPropertyBag* GetPropertyBagStruct() const;
+	
+	/** Returns property descriptor by specified name. */
+	const FPropertyBagPropertyDesc* FindPropertyDescByID(const FGuid ID) const;
+	
+	/** Returns property descriptor by specified ID. */
+	const FPropertyBagPropertyDesc* FindPropertyDescByName(const FName Name) const;
+
+	/** @return const view to the struct that holds the values. NOTE: The returned value/view cannot be serialized, use this to access the struct only temporarily. */
+	FConstStructView GetValue() const { return Value; };
+
+	/** @return const view to the struct that holds the values. NOTE: The returned value/view cannot be serialized, use this to access the struct only temporarily. */
+	FStructView GetMutableValue() const { return Value; };
+	
+	/**
+	 * Getters
+	 * Numeric types (bool, int32, int64, float, double) support type conversion.
+	 */
+
+	TValueOrError<bool, EPropertyBagResult> GetValueBool(const FName Name) const;
+	TValueOrError<uint8, EPropertyBagResult> GetValueByte(const FName Name) const;
+	TValueOrError<int32, EPropertyBagResult> GetValueInt32(const FName Name) const;
+	TValueOrError<int64, EPropertyBagResult> GetValueInt64(const FName Name) const;
+	TValueOrError<float, EPropertyBagResult> GetValueFloat(const FName Name) const;
+	TValueOrError<double, EPropertyBagResult> GetValueDouble(const FName Name) const;
+	TValueOrError<FName, EPropertyBagResult> GetValueName(const FName Name) const;
+	TValueOrError<FString, EPropertyBagResult> GetValueString(const FName Name) const;
+	TValueOrError<FText, EPropertyBagResult> GetValueText(const FName Name) const;
+	TValueOrError<uint8, EPropertyBagResult> GetValueEnum(const FName Name, const UEnum* RequestedEnum) const;
+	TValueOrError<FStructView, EPropertyBagResult> GetValueStruct(const FName Name, const UScriptStruct* RequestedStruct = nullptr) const;
+	TValueOrError<UObject*, EPropertyBagResult> GetValueObject(const FName Name, const UClass* RequestedClass = nullptr) const;
+	TValueOrError<UClass*, EPropertyBagResult> GetValueClass(const FName Name) const;
+
+	/** @return enum value of specified type. */
+	template <typename T>
+	TValueOrError<T, EPropertyBagResult> GetValueEnum(const FName Name) const
+	{
+		static_assert(TIsEnum<T>::Value, "Should only call this with enum types");
+		
+		TValueOrError<uint8, EPropertyBagResult> Result = GetValueEnum(Name, StaticEnum<T>());
+		if (!Result.IsValid())
+		{
+			return MakeError(Result.GetError());
+		}
+		return MakeValue((T)Result.GetValue());
+	}
+	
+	/** @return struct reference of specified type. */
+	template <typename T>
+	TValueOrError<T*, EPropertyBagResult> GetValueStruct(const FName Name) const
+	{
+		TValueOrError<FStructView, EPropertyBagResult> Result = GetValueStruct(Name, TBaseStructure<T>::Get());
+		if (!Result.IsValid())
+		{
+			return MakeError(Result.GetError());
+		}
+		if (T* ValuePtr = Result.GetValue().GetMutablePtr<T>())
+		{
+			return MakeValue(ValuePtr);
+		}
+		return MakeError(EPropertyBagResult::TypeMismatch);
+	}
+	
+	/** @return object pointer value of specified type. */
+	template <typename T>
+	TValueOrError<T*, EPropertyBagResult> GetValueObject(const FName Name) const
+	{
+		static_assert(TIsDerivedFrom<T, UObject>::Value, "Should only call this with object types");
+		
+		TValueOrError<UObject*, EPropertyBagResult> Result = GetValueObject(Name, T::StaticClass());
+		if (!Result.IsValid())
+		{
+			return MakeError(Result.GetError());
+		}
+		if (Result.GetValue() == nullptr)
+		{
+			return MakeValue(nullptr);
+		}
+		if (T* Object = Cast<T>(Result.GetValue()))
+		{
+			return MakeValue(Object);
+		}
+		return MakeError(EPropertyBagResult::TypeMismatch);
+	}
+
+	/**
+	 * Value Setters. A property must exists in that bag before it can be set.  
+	 * Numeric types (bool, int32, int64, float, double) support type conversion.
+	 */
+	EPropertyBagResult SetValueBool(const FName Name, const bool bInValue) const;
+	EPropertyBagResult SetValueByte(const FName Name, const uint8 InValue) const;
+	EPropertyBagResult SetValueInt32(const FName Name, const int32 InValue) const;
+	EPropertyBagResult SetValueInt64(const FName Name, const int64 InValue) const;
+	EPropertyBagResult SetValueFloat(const FName Name, const float InValue) const;
+	EPropertyBagResult SetValueDouble(const FName Name, const double InValue) const;
+	EPropertyBagResult SetValueName(const FName Name, const FName InValue) const;
+	EPropertyBagResult SetValueString(const FName Name, const FString& InValue) const;
+	EPropertyBagResult SetValueText(const FName Name, const FText& InValue) const;
+	EPropertyBagResult SetValueEnum(const FName Name, const uint8 InValue, const UEnum* Enum) const;
+	EPropertyBagResult SetValueStruct(const FName Name, FConstStructView InValue) const;
+	EPropertyBagResult SetValueObject(const FName Name, UObject* InValue) const;
+	EPropertyBagResult SetValueClass(const FName Name, UClass* InValue) const;
+
+	/** Sets enum value specified type. */
+	template <typename T>
+	EPropertyBagResult SetValueEnum(const FName Name, const T InValue) const
+	{
+		static_assert(TIsEnum<T>::Value, "Should only call this with enum types");
+		return SetValueEnum(Name, (uint8)InValue, StaticEnum<T>());
+	}
+
+	/** Sets struct value specified type. */
+	template <typename T>
+	EPropertyBagResult SetValueStruct(const FName Name, const T& InValue) const
+	{
+		return SetValueStruct(Name, FConstStructView::Make(InValue));
+	}
+
+	/** Sets object pointer value specified type. */
+	template <typename T>
+	EPropertyBagResult SetValueObject(const FName Name, T* InValue) const
+	{
+		static_assert(TIsDerivedFrom<T, UObject>::Value, "Should only call this with object types");
+		return SetValueObject(Name, (UObject*)InValue);
+	}
+
+	bool Serialize(FArchive& Ar);
+
+protected:
+	UPROPERTY(EditAnywhere, Category="")
+	FInstancedStruct Value;
+};
+
+template<> struct TStructOpsTypeTraits<FInstancedPropertyBag> : public TStructOpsTypeTraitsBase2<FInstancedPropertyBag>
+{
+	enum
+	{
+		WithSerializer = true
+	};
+};
+
+/**
+ * A script struct that is used to store the value of the property bag instance.
+ * References to UPropertyBag cannot be serialized, instead the array of the properties
+ * is serialized and new class is create on load based on the composition of the properties.
+ *
+ * Note: Should not be used directly.
+ */
+UCLASS(Transient)
+class STRUCTUTILS_API UPropertyBag : public UScriptStruct
+{
+public:
+	GENERATED_BODY()
+
+	/**
+	 * Creates new UPropertyBag struct based on the properties passed in.
+	 * If there are multiple properties that have the same name, only the first one is added.
+	 */
+	static const UPropertyBag* GetOrCreateFromDescs(const TConstArrayView<FPropertyBagPropertyDesc> InPropertyDescs);
+
+	/** Returns property descriptions that specify this struct. */
+	TConstArrayView<FPropertyBagPropertyDesc> GetPropertyDescs() const { return PropertyDescs; }
+
+	/** @return property description based on ID. */
+	const FPropertyBagPropertyDesc* FindPropertyDescByID(const FGuid ID) const;
+
+	/** @return property description based on name. */
+	const FPropertyBagPropertyDesc* FindPropertyDescByName(const FName Name) const;
+
+protected:
+	UPROPERTY()
+	TArray<FPropertyBagPropertyDesc> PropertyDescs;
+
+	friend struct FInstancedPropertyBag;
+};
diff --git a/Engine/Plugins/Experimental/StructUtils/Source/StructUtilsEditor/Private/Customizations/PropertyBagDetails.cpp b/Engine/Plugins/Experimental/StructUtils/Source/StructUtilsEditor/Private/Customizations/PropertyBagDetails.cpp
new file mode 100644
--- /dev/null
+++ b/Engine/Plugins/Experimental/StructUtils/Source/StructUtilsEditor/Private/Customizations/PropertyBagDetails.cpp
@@ -0,0 +1,702 @@
+﻿// Copyright Epic Games, Inc. All Rights Reserved.
+
+#include "PropertyBagDetails.h"
+#include "InstancedStructDetails.h"
+#include "DetailWidgetRow.h"
+#include "DetailLayoutBuilder.h"
+#include "IPropertyUtilities.h"
+#include "IDetailPropertyRow.h"
+#include "IDetailChildrenBuilder.h"
+#include "Widgets/DeclarativeSyntaxSupport.h"
+#include "Widgets/Images/SImage.h"
+#include "Widgets/Text/SInlineEditableTextBlock.h"
+#include "Widgets/Input/SComboButton.h"
+#include "Framework/MultiBox/MultiBoxBuilder.h"
+#include "ScopedTransaction.h"
+#include "Engine/UserDefinedStruct.h"
+#include "Kismet2/BlueprintEditorUtils.h"
+#include "SPinTypeSelector.h"
+#include "PropertyBag.h"
+
+#define LOCTEXT_NAMESPACE "StructUtilsEditor"
+
+////////////////////////////////////
+
+namespace UE::StructUtils::Private
+{
+
+/** @return true property handle holds struct property of type T.  */ 
+template<typename T>
+bool IsScriptStruct(const TSharedPtr<IPropertyHandle>& PropertyHandle)
+{
+	if (!PropertyHandle)
+	{
+		return false;
+	}
+
+	FStructProperty* StructProperty = CastFieldChecked<FStructProperty>(PropertyHandle->GetProperty());
+	return StructProperty && StructProperty->Struct == TBaseStructure<T>::Get();
+}
+
+/** @return property bag struct common to all edited properties. */
+const UPropertyBag* GetCommonBagStruct(TSharedPtr<IPropertyHandle> StructProperty)
+{
+	const UPropertyBag* CommonBagStruct = nullptr;
+
+	if (ensure(IsScriptStruct<FInstancedPropertyBag>(StructProperty)))
+	{
+		StructProperty->EnumerateConstRawData([&CommonBagStruct](const void* RawData, const int32 /*DataIndex*/, const int32 /*NumDatas*/)
+		{
+			if (RawData)
+			{
+				const FInstancedPropertyBag* Bag = static_cast<const FInstancedPropertyBag*>(RawData);
+
+				const UPropertyBag* BagStruct = Bag->GetPropertyBagStruct();
+				if (CommonBagStruct && CommonBagStruct != BagStruct)
+				{
+					// Multiple struct types on the sources - show nothing set
+					CommonBagStruct = nullptr;
+					return false;
+				}
+				CommonBagStruct = BagStruct;
+			}
+
+			return true;
+		});
+	}
+
+	return CommonBagStruct;
+}
+
+/** @return property descriptors of the property bag struct common to all edited properties. */
+TArray<FPropertyBagPropertyDesc> GetCommonPropertyDescs(TSharedPtr<IPropertyHandle> StructProperty)
+{
+	TArray<FPropertyBagPropertyDesc> PropertyDescs;
+	
+	if (const UPropertyBag* BagStruct = GetCommonBagStruct(StructProperty))
+	{
+		PropertyDescs = BagStruct->GetPropertyDescs();
+	}
+	
+	return PropertyDescs;
+}
+
+/** Creates new property bag struct and sets all properties to use it, migrating over old values. */
+void SetPropertyDescs(TSharedPtr<IPropertyHandle> StructProperty, const TConstArrayView<FPropertyBagPropertyDesc> PropertyDescs)
+{
+	if (ensure(IsScriptStruct<FInstancedPropertyBag>(StructProperty)))
+	{
+		// Create new bag struct
+		const UPropertyBag* NewBagStruct = UPropertyBag::GetOrCreateFromDescs(PropertyDescs);
+
+		// Migrate structs to the new type, copying values over.
+		StructProperty->EnumerateRawData([&NewBagStruct](void* RawData, const int32 /*DataIndex*/, const int32 /*NumDatas*/)
+		{
+			if (RawData)
+			{
+				if (FInstancedPropertyBag* Bag = static_cast<FInstancedPropertyBag*>(RawData))
+				{
+					Bag->MigrateToNewBagStruct(NewBagStruct);
+				}
+			}
+
+			return true;
+		});
+	}
+}
+
+/** @return sanitized property name based on the input string. */
+FName GetValidPropertyName(const FString& Name)
+{
+	FName Result;
+	if (!Name.IsEmpty())
+	{
+		if (!FName::IsValidXName(Name, INVALID_OBJECTNAME_CHARACTERS))
+		{
+			Result = MakeObjectNameFromDisplayLabel(Name, NAME_None);
+		}
+		else
+		{
+			Result = FName(Name);
+		}
+	}
+	else
+	{
+		Result = FName(TEXT("Property"));
+	}
+
+	return Result;
+}
+
+/** @return true of the property name is not used yet by the property bag structure common to all edited properties. */
+bool IsUniqueName(const FName NewName, const FName OldName, TSharedPtr<IPropertyHandle> StructProperty)
+{
+	if (NewName == OldName)
+	{
+		return false;
+	}
+	
+	bool bFound = false;
+
+	if (ensure(IsScriptStruct<FInstancedPropertyBag>(StructProperty)))
+	{
+		StructProperty->EnumerateConstRawData([&bFound, NewName](const void* RawData, const int32 /*DataIndex*/, const int32 /*NumDatas*/)
+		{
+			if (const FInstancedPropertyBag* Bag = static_cast<const FInstancedPropertyBag*>(RawData))
+			{
+				if (const UPropertyBag* BagStruct = Bag->GetPropertyBagStruct())
+				{
+					const bool bContains = BagStruct->GetPropertyDescs().ContainsByPredicate([NewName](const FPropertyBagPropertyDesc& Desc)
+					{
+						return Desc.Name == NewName;
+					});
+					if (bContains)
+					{
+						bFound = true;
+						return false; // Stop iterating
+					}
+				}
+			}
+
+			return true;
+		});
+	}
+	
+	return !bFound;
+}
+
+/** @return Blueprint pin type from property descriptor. */
+FEdGraphPinType GetPropertyDescAsPin(const FPropertyBagPropertyDesc& Desc)
+{
+	UEnum* PropertyTypeEnum = StaticEnum<EPropertyBagPropertyType>();
+	check(PropertyTypeEnum);
+	const UEdGraphSchema_K2* Schema = GetDefault<UEdGraphSchema_K2>();
+	check(Schema);
+
+	FEdGraphPinType PinType;
+	PinType.PinSubCategory = NAME_None;
+
+	switch (Desc.ValueType)
+	{
+	case EPropertyBagPropertyType::Bool:
+		PinType.PinCategory = UEdGraphSchema_K2::PC_Boolean;
+		break;
+	case EPropertyBagPropertyType::Byte:
+		PinType.PinCategory = UEdGraphSchema_K2::PC_Byte;
+		break;
+	case EPropertyBagPropertyType::Int32:
+		PinType.PinCategory = UEdGraphSchema_K2::PC_Int;
+		break;
+	case EPropertyBagPropertyType::Int64:
+		PinType.PinCategory = UEdGraphSchema_K2::PC_Int64;
+		break;
+	case EPropertyBagPropertyType::Float:
+		PinType.PinCategory = UEdGraphSchema_K2::PC_Real;
+		PinType.PinSubCategory = UEdGraphSchema_K2::PC_Float;
+		break;
+	case EPropertyBagPropertyType::Double:
+		PinType.PinCategory = UEdGraphSchema_K2::PC_Real;
+		PinType.PinSubCategory = UEdGraphSchema_K2::PC_Double;
+		break;
+	case EPropertyBagPropertyType::Name:
+		PinType.PinCategory = UEdGraphSchema_K2::PC_Name;
+		break;
+	case EPropertyBagPropertyType::String:
+		PinType.PinCategory = UEdGraphSchema_K2::PC_String;
+		break;
+	case EPropertyBagPropertyType::Text:
+		PinType.PinCategory = UEdGraphSchema_K2::PC_Text;
+		break;
+	case EPropertyBagPropertyType::Enum:
+		// @todo: some pin coloring is not correct due to this (byte-as-enum vs enum). 
+		PinType.PinCategory = UEdGraphSchema_K2::PC_Enum;
+		PinType.PinSubCategoryObject = const_cast<UObject*>(Desc.ValueTypeObject.Get());
+		break;
+	case EPropertyBagPropertyType::Struct:
+		PinType.PinCategory = UEdGraphSchema_K2::PC_Struct;
+		PinType.PinSubCategoryObject = const_cast<UObject*>(Desc.ValueTypeObject.Get());
+		break;
+	case EPropertyBagPropertyType::Object:
+		PinType.PinCategory = UEdGraphSchema_K2::PC_Object;
+		PinType.PinSubCategoryObject = const_cast<UObject*>(Desc.ValueTypeObject.Get());
+		break;
+	case EPropertyBagPropertyType::SoftObject:
+		PinType.PinCategory = UEdGraphSchema_K2::PC_SoftObject;
+		PinType.PinSubCategoryObject = const_cast<UObject*>(Desc.ValueTypeObject.Get());
+		break;
+	case EPropertyBagPropertyType::Class:
+		PinType.PinCategory = UEdGraphSchema_K2::PC_Class;
+		PinType.PinSubCategoryObject = const_cast<UObject*>(Desc.ValueTypeObject.Get());
+		break;
+	case EPropertyBagPropertyType::SoftClass:
+		PinType.PinCategory = UEdGraphSchema_K2::PC_SoftClass;
+		PinType.PinSubCategoryObject = const_cast<UObject*>(Desc.ValueTypeObject.Get());
+		break;
+	default:
+		ensureMsgf(false, TEXT("Unhandled value type %s"), *UEnum::GetValueAsString(Desc.ValueType));
+		break;
+	}
+
+	return PinType;
+}
+
+/** Sets property descriptor based on a Blueprint pin type. */
+void SetPropertyDescFromPin(FPropertyBagPropertyDesc& Desc, const FEdGraphPinType& PinType)
+{
+	const UEdGraphSchema_K2* Schema = GetDefault<UEdGraphSchema_K2>();
+	check(Schema);
+
+	if (PinType.PinCategory == UEdGraphSchema_K2::PC_Boolean)
+	{
+		Desc.ValueType = EPropertyBagPropertyType::Bool;
+		Desc.ValueTypeObject = nullptr;
+	}
+	else if (PinType.PinCategory == UEdGraphSchema_K2::PC_Byte)
+	{
+		if (UEnum* Enum = Cast<UEnum>(PinType.PinSubCategoryObject))
+		{
+			Desc.ValueType = EPropertyBagPropertyType::Enum;
+			Desc.ValueTypeObject = PinType.PinSubCategoryObject.Get();
+		}
+		else
+		{
+			Desc.ValueType = EPropertyBagPropertyType::Byte;
+			Desc.ValueTypeObject = nullptr;
+		}
+	}
+	else if (PinType.PinCategory == UEdGraphSchema_K2::PC_Int)
+	{
+		Desc.ValueType = EPropertyBagPropertyType::Int32;
+		Desc.ValueTypeObject = nullptr;
+	}
+	else if (PinType.PinCategory == UEdGraphSchema_K2::PC_Int64)
+	{
+		Desc.ValueType = EPropertyBagPropertyType::Int64;
+		Desc.ValueTypeObject = nullptr;
+	}
+	else if (PinType.PinCategory == UEdGraphSchema_K2::PC_Real)
+	{
+		if (PinType.PinSubCategory == UEdGraphSchema_K2::PC_Float)
+		{
+			Desc.ValueType = EPropertyBagPropertyType::Float;
+			Desc.ValueTypeObject = nullptr;
+		}
+		else if (PinType.PinSubCategory == UEdGraphSchema_K2::PC_Double)
+		{
+			Desc.ValueType = EPropertyBagPropertyType::Double;
+			Desc.ValueTypeObject = nullptr;
+		}		
+	}
+	else if (PinType.PinCategory == UEdGraphSchema_K2::PC_Name)
+	{
+		Desc.ValueType = EPropertyBagPropertyType::Name;
+		Desc.ValueTypeObject = nullptr;
+	}
+	else if (PinType.PinCategory == UEdGraphSchema_K2::PC_String)
+	{
+		Desc.ValueType = EPropertyBagPropertyType::String;
+		Desc.ValueTypeObject = nullptr;
+	}
+	else if (PinType.PinCategory == UEdGraphSchema_K2::PC_Text)
+	{
+		Desc.ValueType = EPropertyBagPropertyType::Text;
+		Desc.ValueTypeObject = nullptr;
+	}
+	else if (PinType.PinCategory == UEdGraphSchema_K2::PC_Enum)
+	{
+		Desc.ValueType = EPropertyBagPropertyType::Enum;
+		Desc.ValueTypeObject = PinType.PinSubCategoryObject.Get();
+	}
+	else if (PinType.PinCategory == UEdGraphSchema_K2::PC_Struct)
+	{
+		Desc.ValueType = EPropertyBagPropertyType::Struct;
+		Desc.ValueTypeObject = PinType.PinSubCategoryObject.Get();
+	}
+	else if (PinType.PinCategory == UEdGraphSchema_K2::PC_Object)
+	{
+		Desc.ValueType = EPropertyBagPropertyType::Object;
+		Desc.ValueTypeObject = PinType.PinSubCategoryObject.Get();
+	}
+	else if (PinType.PinCategory == UEdGraphSchema_K2::PC_SoftObject)
+	{
+		Desc.ValueType = EPropertyBagPropertyType::SoftObject;
+		Desc.ValueTypeObject = PinType.PinSubCategoryObject.Get();
+	}
+	else if (PinType.PinCategory == UEdGraphSchema_K2::PC_Class)
+	{
+		Desc.ValueType = EPropertyBagPropertyType::Class;
+		Desc.ValueTypeObject = PinType.PinSubCategoryObject.Get();
+	}
+	else if (PinType.PinCategory == UEdGraphSchema_K2::PC_SoftClass)
+	{
+		Desc.ValueType = EPropertyBagPropertyType::SoftClass;
+		Desc.ValueTypeObject = PinType.PinSubCategoryObject.Get();
+	}
+	else
+	{
+		ensureMsgf(false, TEXT("Unhandled pin category %s"), *PinType.PinCategory.ToString());
+	}
+}
+
+template<typename TFunc>
+void ApplyChangesToPropertyDescs(const FText& SessionName, TSharedPtr<IPropertyHandle> StructProperty, IPropertyUtilities* PropUtils, TFunc&& Function)
+{
+	FScopedTransaction Transaction(SessionName);
+	TArray<FPropertyBagPropertyDesc> PropertyDescs = GetCommonPropertyDescs(StructProperty);
+	StructProperty->NotifyPreChange();
+
+	Function(PropertyDescs);
+
+	SetPropertyDescs(StructProperty, PropertyDescs);
+	
+	StructProperty->NotifyPostChange(EPropertyChangeType::ValueSet);
+	StructProperty->NotifyFinishedChangingProperties();
+	if (PropUtils)
+	{
+		PropUtils->ForceRefresh();
+	}
+}
+
+
+} // UE::StructUtils::Private
+
+
+//----------------------------------------------------------------//
+//  FPropertyBagInstanceDataDetails
+//  - StructProperty is FInstancedPropertyBag
+//  - ValueProperty is FInstancedStruct
+//  - ChildPropertyHandle a child property if the ValueProperty  
+//----------------------------------------------------------------//
+
+class FPropertyBagInstanceDataDetails : public FInstancedStructDataDetails
+{
+public:
+	FPropertyBagInstanceDataDetails(TSharedPtr<IPropertyHandle> InValueProperty, TSharedPtr<IPropertyHandle> InStructProperty, IPropertyUtilities* InPropUtils)
+		: FInstancedStructDataDetails(InValueProperty)
+		, StructProperty(InStructProperty)
+		, PropUtils(InPropUtils)
+	{
+	}
+
+	virtual void OnChildRowAdded(IDetailPropertyRow& ChildRow) override
+	{
+		TSharedPtr<SWidget> NameWidget;
+		TSharedPtr<SWidget> ValueWidget;
+		FDetailWidgetRow Row;
+		ChildRow.GetDefaultWidgets(NameWidget, ValueWidget, Row);
+
+		// Inline editable name
+		TSharedPtr<SInlineEditableTextBlock> InlineWidget = SNew(SInlineEditableTextBlock)
+			.Text(this, &FPropertyBagInstanceDataDetails::GetPropertyName, ChildRow.GetPropertyHandle())
+			.Font(IDetailLayoutBuilder::GetDetailFont())
+			.OnVerifyTextChanged(this, &FPropertyBagInstanceDataDetails::OnValidatePropertyName, ChildRow.GetPropertyHandle())
+			.OnTextCommitted(this, &FPropertyBagInstanceDataDetails::OnChangePropertyName, ChildRow.GetPropertyHandle())
+			.MultiLine(false);
+
+		ChildRow
+			.CustomWidget(/*bShowChildren*/true)
+			.NameContent()
+			[
+				SNew(SComboButton)
+				.OnGetMenuContent(this, &FPropertyBagInstanceDataDetails::OnPropertyNameContent, ChildRow.GetPropertyHandle(), InlineWidget)
+				.ContentPadding(2)
+				.ButtonStyle(FEditorStyle::Get(), "HoverHintOnly")
+				.ButtonContent()
+				[
+					InlineWidget.ToSharedRef()
+				]
+			]
+			.ValueContent()
+			[
+				ValueWidget.ToSharedRef()
+			];
+	}
+
+
+	FEdGraphPinType OnGetPinInfo(TSharedPtr<IPropertyHandle> ChildPropertyHandle) const
+	{
+		TArray<FPropertyBagPropertyDesc> PropertyDescs = UE::StructUtils::Private::GetCommonPropertyDescs(StructProperty);
+
+		const FProperty* Property = ChildPropertyHandle->GetProperty();
+		if (FPropertyBagPropertyDesc* Desc = PropertyDescs.FindByPredicate([Property](const FPropertyBagPropertyDesc& Desc){ return Desc.CachedProperty == Property; }))
+		{
+			return UE::StructUtils::Private::GetPropertyDescAsPin(*Desc);
+		}
+		
+		return FEdGraphPinType();
+	}
+
+	void PinInfoChanged(const FEdGraphPinType& PinType, TSharedPtr<IPropertyHandle> ChildPropertyHandle) const
+	{
+		UE::StructUtils::Private::ApplyChangesToPropertyDescs(
+			LOCTEXT("OnPropertyTypeChanged", "Change Property Type"), StructProperty, PropUtils,
+			[&PinType, &ChildPropertyHandle](TArray<FPropertyBagPropertyDesc>& PropertyDescs)
+			{
+				// Find and change struct type
+				const FProperty* Property = ChildPropertyHandle->GetProperty();
+				if (FPropertyBagPropertyDesc* Desc = PropertyDescs.FindByPredicate([Property](const FPropertyBagPropertyDesc& Desc){ return Desc.CachedProperty == Property; }))
+				{
+					UE::StructUtils::Private::SetPropertyDescFromPin(*Desc, PinType);
+				}
+			});
+	}
+
+	static bool CanHaveMemberVariableOfType(const FEdGraphPinType& PinType)
+	{
+		if ((PinType.PinCategory == UEdGraphSchema_K2::PC_Struct))
+		{
+			if (const UObject* TypeObject = PinType.PinSubCategoryObject.Get())
+			{
+				if (TypeObject->IsA<UUserDefinedStruct>())
+				{
+					return false;
+				}
+			}
+		}
+		else if ((PinType.PinCategory == UEdGraphSchema_K2::PC_Exec) 
+			|| (PinType.PinCategory == UEdGraphSchema_K2::PC_Wildcard)
+			|| (PinType.PinCategory == UEdGraphSchema_K2::PC_MCDelegate)
+			|| (PinType.PinCategory == UEdGraphSchema_K2::PC_Delegate)
+			|| (PinType.PinCategory == UEdGraphSchema_K2::PC_Interface))
+		{
+			return false;
+		}
+		return true;
+	}
+	
+	void GetFilteredVariableTypeTree(TArray<TSharedPtr<UEdGraphSchema_K2::FPinTypeTreeInfo>>& TypeTree, ETypeTreeFilter TypeTreeFilter) const
+	{
+		check(GetDefault<UEdGraphSchema_K2>());
+		GetDefault<UEdGraphSchema_K2>()->GetVariableTypeTree(TypeTree, TypeTreeFilter);
+
+		// Filter
+		for (TSharedPtr<UEdGraphSchema_K2::FPinTypeTreeInfo>& PinType : TypeTree)
+		{
+			if (!PinType.IsValid())
+			{
+				return;
+			}
+
+			for (int32 ChildIndex = 0; ChildIndex < PinType->Children.Num(); )
+			{
+				TSharedPtr<UEdGraphSchema_K2::FPinTypeTreeInfo> Child = PinType->Children[ChildIndex];
+				if (Child.IsValid())
+				{
+					if (!CanHaveMemberVariableOfType(Child->GetPinType(/*bForceLoadSubCategoryObject*/false)))
+					{
+						PinType->Children.RemoveAt(ChildIndex);
+						continue;
+					}
+				}
+				++ChildIndex;
+			}
+		}			
+	}
+
+	
+	FText GetPropertyName(TSharedPtr<IPropertyHandle> ChildPropertyHandle) const
+	{
+		return FText::FromName(ChildPropertyHandle->GetProperty()->GetFName());
+	}
+
+	bool OnValidatePropertyName(const FText& InText, FText& OutErrorMessage, TSharedPtr<IPropertyHandle> ChildPropertyHandle) const
+	{
+		const FName NewName = UE::StructUtils::Private::GetValidPropertyName(InText.ToString());
+		bool bResult = UE::StructUtils::Private::IsUniqueName(NewName, ChildPropertyHandle->GetProperty()->GetFName(), StructProperty);
+		if (!bResult)
+		{
+			OutErrorMessage = LOCTEXT("MustBeUniqueName", "Property must have unique name");
+		}
+		return bResult;
+	}
+
+	void OnChangePropertyName(const FText& InNewText, ETextCommit::Type InCommitType, TSharedPtr<IPropertyHandle> ChildPropertyHandle) const
+	{
+		if (InCommitType == ETextCommit::OnCleared)
+		{
+			return;
+		}
+
+		const FName NewName = UE::StructUtils::Private::GetValidPropertyName(InNewText.ToString());
+		if (!UE::StructUtils::Private::IsUniqueName(NewName, ChildPropertyHandle->GetProperty()->GetFName(), StructProperty))
+		{
+			return;
+		}
+
+		UE::StructUtils::Private::ApplyChangesToPropertyDescs(
+			LOCTEXT("OnPropertyNameChanged", "Change Property Name"), StructProperty, PropUtils,
+			[&NewName, &ChildPropertyHandle](TArray<FPropertyBagPropertyDesc>& PropertyDescs)
+			{
+				const FProperty* Property = ChildPropertyHandle->GetProperty();
+				if (FPropertyBagPropertyDesc* Desc = PropertyDescs.FindByPredicate([Property](const FPropertyBagPropertyDesc& Desc){ return Desc.CachedProperty == Property; }))
+				{
+					Desc->Name = NewName;
+				}
+			});
+	}
+
+	void OnRemoveProperty(TSharedPtr<IPropertyHandle> ChildPropertyHandle) const
+	{
+		UE::StructUtils::Private::ApplyChangesToPropertyDescs(
+			LOCTEXT("OnPropertyRemoved", "Remove Property"), StructProperty, PropUtils,
+			[&ChildPropertyHandle](TArray<FPropertyBagPropertyDesc>& PropertyDescs)
+			{
+				const FProperty* Property = ChildPropertyHandle->GetProperty();
+				PropertyDescs.RemoveAll([Property](const FPropertyBagPropertyDesc& Desc){ return Desc.CachedProperty == Property; });
+			});
+	}
+
+	void OnMoveProperty(TSharedPtr<IPropertyHandle> ChildPropertyHandle, int32 Delta) const
+	{
+		UE::StructUtils::Private::ApplyChangesToPropertyDescs(
+			LOCTEXT("OnPropertyMoved", "Move Property"), StructProperty, PropUtils,
+			[&ChildPropertyHandle, &Delta](TArray<FPropertyBagPropertyDesc>& PropertyDescs)
+			{
+				// Move
+				if (PropertyDescs.Num() > 1)
+				{
+					const FProperty* Property = ChildPropertyHandle->GetProperty();
+					const int32 PropertyIndex = PropertyDescs.IndexOfByPredicate([Property](const FPropertyBagPropertyDesc& Desc){ return Desc.CachedProperty == Property; });
+					if (PropertyIndex != INDEX_NONE)
+					{
+						const int32 NewPropertyIndex = FMath::Clamp(PropertyIndex + Delta, 0, PropertyDescs.Num() - 1);
+						PropertyDescs.Swap(PropertyIndex, NewPropertyIndex);
+					}
+				}
+			});
+	}
+
+	TSharedRef<SWidget> OnPropertyNameContent(TSharedPtr<IPropertyHandle> ChildPropertyHandle, TSharedPtr<SInlineEditableTextBlock> InlineWidget) const
+	{
+		constexpr bool bInShouldCloseWindowAfterMenuSelection = true;
+		FMenuBuilder MenuBuilder(bInShouldCloseWindowAfterMenuSelection, nullptr);
+
+		MenuBuilder.AddMenuEntry(
+			FUIAction(),
+			SNew(SPinTypeSelector, FGetPinTypeTree::CreateSP(this, &FPropertyBagInstanceDataDetails::GetFilteredVariableTypeTree))
+				.TargetPinType(this, &FPropertyBagInstanceDataDetails::OnGetPinInfo, ChildPropertyHandle)
+				.OnPinTypeChanged(this, &FPropertyBagInstanceDataDetails::PinInfoChanged, ChildPropertyHandle)
+				.Schema(GetDefault<UEdGraphSchema_K2>())
+				.bAllowArrays(false)
+				.TypeTreeFilter(ETypeTreeFilter::None)
+				.Font( IDetailLayoutBuilder::GetDetailFont() )
+			);
+
+		MenuBuilder.AddSeparator();
+		
+		MenuBuilder.AddMenuEntry(
+			LOCTEXT("Rename", "Rename"),
+			LOCTEXT("Rename_ToolTip", "Rename property"),
+			FSlateIcon(FAppStyle::GetAppStyleSetName(), "Icons.Edit"),
+			FUIAction(FExecuteAction::CreateLambda([InlineWidget]()  { InlineWidget->EnterEditingMode(); }))
+		);
+
+		MenuBuilder.AddMenuEntry(
+			LOCTEXT("Remove", "Remove"),
+			LOCTEXT("Remove_ToolTip", "Remove property"),
+			FSlateIcon(FAppStyle::GetAppStyleSetName(), "Icons.Delete"),
+			FUIAction(FExecuteAction::CreateSP(this, &FPropertyBagInstanceDataDetails::OnRemoveProperty, ChildPropertyHandle))
+		);
+		
+		MenuBuilder.AddSeparator();
+
+		MenuBuilder.AddMenuEntry(
+			LOCTEXT("MoveUp", "Move Up"),
+			LOCTEXT("MoveUp_ToolTip", "Move property up in the list"),
+			FSlateIcon(FAppStyle::GetAppStyleSetName(), "Icons.ArrowUp"),
+			FUIAction(FExecuteAction::CreateSP(this, &FPropertyBagInstanceDataDetails::OnMoveProperty, ChildPropertyHandle, -1))
+		);
+
+		MenuBuilder.AddMenuEntry(
+			LOCTEXT("MoveDown", "Move Down"),
+			LOCTEXT("MoveDown_ToolTip", "Move property down in the list"),
+			FSlateIcon(FAppStyle::GetAppStyleSetName(), "Icons.ArrowDown"),
+			FUIAction(FExecuteAction::CreateSP(this, &FPropertyBagInstanceDataDetails::OnMoveProperty, ChildPropertyHandle, +1))
+		);
+
+		return MenuBuilder.MakeWidget();
+	}
+
+	TSharedPtr<IPropertyHandle> StructProperty;
+	IPropertyUtilities* PropUtils = nullptr;
+};
+
+
+////////////////////////////////////
+
+TSharedRef<IPropertyTypeCustomization> FPropertyBagDetails::MakeInstance()
+{
+	return MakeShared<FPropertyBagDetails>();
+}
+
+void FPropertyBagDetails::CustomizeHeader(TSharedRef<class IPropertyHandle> StructPropertyHandle, class FDetailWidgetRow& HeaderRow, IPropertyTypeCustomizationUtils& StructCustomizationUtils)
+{
+	PropUtils = StructCustomizationUtils.GetPropertyUtilities().Get();
+	
+	StructProperty = StructPropertyHandle;
+	check(StructProperty);
+	
+	ValueProperty = StructProperty->GetChildHandle(TEXT("Value"));
+	check(ValueProperty);
+
+
+	HeaderRow
+		.NameContent()
+		[
+			StructPropertyHandle->CreatePropertyNameWidget()
+		]
+		.ValueContent()
+		.VAlign(VAlign_Center)
+		[
+			SNew(SHorizontalBox)
+			+SHorizontalBox::Slot()
+			.AutoWidth()
+			[
+				SNew(SButton)
+				.VAlign(VAlign_Center)
+				.HAlign(HAlign_Center)
+				.ButtonStyle(FEditorStyle::Get(), "HoverHintOnly")
+				.ToolTipText(LOCTEXT("AddProperty_Tooltip", "Add new property"))
+				.OnClicked(this, &FPropertyBagDetails::OnAddProperty)
+				[
+					SNew(SImage)
+					.Image(FEditorStyle::GetBrush("Icons.PlusCircle"))
+				]
+			]
+		]
+		.ShouldAutoExpand(true);
+}
+
+void FPropertyBagDetails::CustomizeChildren(TSharedRef<class IPropertyHandle> StructPropertyHandle, class IDetailChildrenBuilder& StructBuilder, IPropertyTypeCustomizationUtils& StructCustomizationUtils)
+{
+	// Show the Value (FInstancedStruct) as child rows. 
+	TSharedRef<FPropertyBagInstanceDataDetails> InstanceDetails = MakeShareable(new FPropertyBagInstanceDataDetails(ValueProperty, StructProperty, PropUtils));
+	StructBuilder.AddCustomBuilder(InstanceDetails);
+}
+
+FReply FPropertyBagDetails::OnAddProperty() const
+{
+	constexpr int32 MaxIterations = 100;
+	FName NewName(TEXT("NewProperty"));
+	int32 Number = 1;
+	while (!UE::StructUtils::Private::IsUniqueName(NewName, FName(), StructProperty) && Number < MaxIterations)
+	{
+		Number++;
+		NewName.SetNumber(Number);
+	}
+	if (Number == MaxIterations)
+	{
+		return FReply::Handled();
+	}
+
+	UE::StructUtils::Private::ApplyChangesToPropertyDescs(
+		LOCTEXT("OnPropertyAdded", "Add Property"), StructProperty, PropUtils,
+		[&NewName](TArray<FPropertyBagPropertyDesc>& PropertyDescs)
+		{
+			PropertyDescs.Emplace(NewName, EPropertyBagPropertyType::Bool);
+		});
+	
+	return FReply::Handled();
+}
+
+#undef LOCTEXT_NAMESPACE
diff --git a/Engine/Plugins/Experimental/StructUtils/Source/StructUtilsEditor/Private/StructUtilsEditorModule.cpp b/Engine/Plugins/Experimental/StructUtils/Source/StructUtilsEditor/Private/StructUtilsEditorModule.cpp
--- a/Engine/Plugins/Experimental/StructUtils/Source/StructUtilsEditor/Private/StructUtilsEditorModule.cpp
+++ b/Engine/Plugins/Experimental/StructUtils/Source/StructUtilsEditor/Private/StructUtilsEditorModule.cpp
@@ -3,6 +3,7 @@
 #include "StructUtilsEditorModule.h"
 #include "Modules/ModuleManager.h"
 #include "InstancedStructDetails.h"
+#include "PropertyBagDetails.h"
 
 #define LOCTEXT_NAMESPACE "StructUtilsEditor"
 
@@ -13,6 +14,7 @@ void FStructUtilsEditorModule::StartupModule()
 	// Register the details customizer
 	FPropertyEditorModule& PropertyModule = FModuleManager::LoadModuleChecked<FPropertyEditorModule>("PropertyEditor");
 	PropertyModule.RegisterCustomPropertyTypeLayout("InstancedStruct", FOnGetPropertyTypeCustomizationInstance::CreateStatic(&FInstancedStructDetails::MakeInstance));
+	PropertyModule.RegisterCustomPropertyTypeLayout("InstancedPropertyBag", FOnGetPropertyTypeCustomizationInstance::CreateStatic(&FPropertyBagDetails::MakeInstance));
 	PropertyModule.NotifyCustomizationModuleChanged();
 }
 
@@ -23,6 +25,7 @@ void FStructUtilsEditorModule::ShutdownModule()
 	{
 		FPropertyEditorModule& PropertyModule = FModuleManager::LoadModuleChecked<FPropertyEditorModule>("PropertyEditor");
 		PropertyModule.UnregisterCustomPropertyTypeLayout("InstancedStruct");
+		PropertyModule.UnregisterCustomPropertyTypeLayout("InstancedPropertyBag");
 		PropertyModule.NotifyCustomizationModuleChanged();
 	}
 }
diff --git a/Engine/Plugins/Experimental/StructUtils/Source/StructUtilsEditor/Public/PropertyBagDetails.h b/Engine/Plugins/Experimental/StructUtils/Source/StructUtilsEditor/Public/PropertyBagDetails.h
new file mode 100644
--- /dev/null
+++ b/Engine/Plugins/Experimental/StructUtils/Source/StructUtilsEditor/Public/PropertyBagDetails.h
@@ -0,0 +1,32 @@
+﻿// Copyright Epic Games, Inc. All Rights Reserved.
+
+#pragma once
+
+#include "IPropertyTypeCustomization.h"
+
+class IPropertyHandle;
+
+/**
+ * Type customization for FPropertyBag.
+ */
+class STRUCTUTILSEDITOR_API FPropertyBagDetails : public IPropertyTypeCustomization
+{
+public:
+	/** Makes a new instance of this detail layout class for a specific detail view requesting it */
+	static TSharedRef<IPropertyTypeCustomization> MakeInstance();
+
+protected:
+	/** IPropertyTypeCustomization interface */
+	virtual void CustomizeHeader(TSharedRef<IPropertyHandle> StructPropertyHandle, class FDetailWidgetRow& HeaderRow, IPropertyTypeCustomizationUtils& StructCustomizationUtils) override;
+	virtual void CustomizeChildren(TSharedRef<IPropertyHandle> StructPropertyHandle, class IDetailChildrenBuilder& StructBuilder, IPropertyTypeCustomizationUtils& StructCustomizationUtils) override;
+
+private:
+
+	FReply OnAddProperty() const;
+	
+	/** Handle to the struct property being edited */
+	TSharedPtr<IPropertyHandle> StructProperty;
+	TSharedPtr<IPropertyHandle> ValueProperty;
+
+	class IPropertyUtilities* PropUtils = nullptr;
+};
diff --git a/Engine/Plugins/Experimental/StructUtils/Source/StructUtilsEditor/StructUtilsEditor.Build.cs b/Engine/Plugins/Experimental/StructUtils/Source/StructUtilsEditor/StructUtilsEditor.Build.cs
--- a/Engine/Plugins/Experimental/StructUtils/Source/StructUtilsEditor/StructUtilsEditor.Build.cs
+++ b/Engine/Plugins/Experimental/StructUtils/Source/StructUtilsEditor/StructUtilsEditor.Build.cs
@@ -40,6 +40,7 @@ namespace UnrealBuildTool.Rules
 				"PropertyEditor",
 				"AIGraph",
 				"ToolMenus",
+				"BlueprintGraph",
 			}
 			);
 		}
diff --git a/Engine/Plugins/Experimental/StructUtils/Source/StructUtilsTestSuite/Private/PropertyBagTest.cpp b/Engine/Plugins/Experimental/StructUtils/Source/StructUtilsTestSuite/Private/PropertyBagTest.cpp
new file mode 100644
--- /dev/null
+++ b/Engine/Plugins/Experimental/StructUtils/Source/StructUtilsTestSuite/Private/PropertyBagTest.cpp
@@ -0,0 +1,263 @@
+﻿// Copyright Epic Games, Inc. All Rights Reserved.
+
+#include "AITestsCommon.h"
+#include "PropertyBag.h"
+#include "StructUtilsTestTypes.h"
+
+#define LOCTEXT_NAMESPACE "StructUtilsTests"
+
+PRAGMA_DISABLE_OPTIMIZATION
+
+namespace FPropertyBagTest
+{
+
+struct FTest_CreatePropertyBag : FAITestBase
+{
+	virtual bool InstantTest() override
+	{
+		static const FName IsHotName(TEXT("bIsHot"));
+		static const FName TemperatureName(TEXT("Temperature"));
+		static const FName CountName(TEXT("Count"));
+
+
+		FInstancedPropertyBag Bag;
+
+		Bag.AddProperty(IsHotName, EPropertyBagPropertyType::Bool);
+		AITEST_TRUE(TEXT("Should have bIsHot property"), Bag.FindPropertyDescByName(IsHotName) != nullptr);
+		AITEST_TRUE(TEXT("Set bIsHot should succeed"), Bag.SetValueBool(IsHotName, true) == EPropertyBagResult::Success);
+
+		// Amend the bag with new properties.
+		Bag.AddProperties({
+			{ TemperatureName, EPropertyBagPropertyType::Float },
+			{ CountName, EPropertyBagPropertyType::Int32 }
+			});
+		AITEST_TRUE(TEXT("Set Temperature should succeed"), Bag.SetValueFloat(TemperatureName, 451.0f) == EPropertyBagResult::Success);
+		AITEST_TRUE(TEXT("Set Count should succeed"), Bag.SetValueFloat(CountName, 42) == EPropertyBagResult::Success);
+
+		Bag.RemovePropertyByName(IsHotName);
+		AITEST_TRUE(TEXT("Should not have bIsHot property"), Bag.FindPropertyDescByName(IsHotName) == nullptr);
+		AITEST_TRUE(TEXT("Set bIsHot should not succeed"), Bag.SetValueBool(IsHotName, true) != EPropertyBagResult::Success);
+		
+		return true;
+	}
+};
+IMPLEMENT_AI_INSTANT_TEST(FTest_CreatePropertyBag, "System.StructUtils.PropertyBag.CreateBag");
+
+struct FTest_MovePropertyBag : FAITestBase
+{
+	virtual bool InstantTest() override
+	{
+		static const FName IsHotName(TEXT("bIsHot"));
+		static const FName TemperatureName(TEXT("Temperature"));
+
+		FInstancedPropertyBag Bag;
+
+		Bag.AddProperty(IsHotName, EPropertyBagPropertyType::Bool);
+		AITEST_TRUE(TEXT("Bag should have bIsHot property"), Bag.FindPropertyDescByName(IsHotName) != nullptr);
+		AITEST_TRUE(TEXT("Bag set bIsHot should succeed"), Bag.SetValueBool(IsHotName, true) == EPropertyBagResult::Success);
+
+		FInstancedPropertyBag Bag2(Bag);
+		Bag2.AddProperty(TemperatureName, EPropertyBagPropertyType::Float);
+		AITEST_TRUE(TEXT("Bag should have bIsHot property"), Bag.FindPropertyDescByName(IsHotName) != nullptr);
+		AITEST_TRUE(TEXT("Bag2 should have bIsHot property"), Bag2.FindPropertyDescByName(IsHotName) != nullptr);
+		AITEST_TRUE(TEXT("Bag2 should have Temperature property"), Bag2.FindPropertyDescByName(TemperatureName) != nullptr);
+
+		FInstancedPropertyBag Bag3(MoveTemp(Bag));
+		AITEST_TRUE(TEXT("Bag should not have bIsHot property"), Bag.FindPropertyDescByName(IsHotName) == nullptr);
+		AITEST_TRUE(TEXT("Bag3 should have bIsHot property"), Bag3.FindPropertyDescByName(IsHotName) != nullptr);
+
+		Bag = Bag2;
+		AITEST_TRUE(TEXT("Bag should have bIsHot property"), Bag.FindPropertyDescByName(IsHotName) != nullptr);
+		AITEST_TRUE(TEXT("Bag should have Temperature property"), Bag.FindPropertyDescByName(TemperatureName) != nullptr);
+		AITEST_TRUE(TEXT("Bag2 should have bIsHot property"), Bag2.FindPropertyDescByName(IsHotName) != nullptr);
+		AITEST_TRUE(TEXT("Bag2 should have Temperature property"), Bag2.FindPropertyDescByName(TemperatureName) != nullptr);
+
+		Bag = MoveTemp(Bag2);
+		AITEST_TRUE(TEXT("Bag should have bIsHot property"), Bag.FindPropertyDescByName(IsHotName) != nullptr);
+		AITEST_TRUE(TEXT("Bag should have Temperature property"), Bag.FindPropertyDescByName(TemperatureName) != nullptr);
+		AITEST_TRUE(TEXT("Bag2 should not have bIsHot property"), Bag2.FindPropertyDescByName(IsHotName) == nullptr);
+		AITEST_TRUE(TEXT("Bag2 should not have Temperature property"), Bag2.FindPropertyDescByName(TemperatureName) == nullptr);
+
+		return true;
+	}
+};
+IMPLEMENT_AI_INSTANT_TEST(FTest_MovePropertyBag, "System.StructUtils.PropertyBag.MoveBag");
+
+struct FTest_MigrateProperty : FAITestBase
+{
+	virtual bool InstantTest() override
+	{
+		static const FName TemperatureName(TEXT("Temperature"));
+
+		FInstancedPropertyBag Bag;
+		Bag.AddProperty(TemperatureName, EPropertyBagPropertyType::Float);
+		AITEST_TRUE(TEXT("Bag should have Temperature property"), Bag.FindPropertyDescByName(TemperatureName) != nullptr);
+		
+		AITEST_TRUE(TEXT("Bag set Temperature as float should succeed"), Bag.SetValueFloat(TemperatureName, 451.0f) == EPropertyBagResult::Success);
+		TValueOrError<float, EPropertyBagResult> FloatRes = Bag.GetValueFloat(TemperatureName);
+		AITEST_TRUE(TEXT("Bag Temperature as float should be 451"), FloatRes.IsValid() && FMath::IsNearlyEqual(FloatRes.GetValue(), 451.0f));
+
+		AITEST_TRUE(TEXT("Bag set Temperature as int should succeed"), Bag.SetValueInt32(TemperatureName, 451) == EPropertyBagResult::Success);
+		FloatRes = Bag.GetValueFloat(TemperatureName);
+		AITEST_TRUE(TEXT("Bag Temperature as float should be 451"), FloatRes.IsValid() && FMath::IsNearlyEqual(FloatRes.GetValue(), 451.0f));
+		TValueOrError<int64, EPropertyBagResult> Int64Res = Bag.GetValueInt64(TemperatureName);
+		AITEST_TRUE(TEXT("Bag Temperature as int64 should be 451"), Int64Res.IsValid() && Int64Res.GetValue() == 451);
+
+		Bag.AddProperty(TemperatureName, EPropertyBagPropertyType::Int32);
+		const FPropertyBagPropertyDesc* TempDesc = Bag.FindPropertyDescByName(TemperatureName);
+		AITEST_TRUE(TEXT("Temperature property should be int32"), TempDesc != nullptr && TempDesc->ValueType == EPropertyBagPropertyType::Int32);
+
+		TValueOrError<int32, EPropertyBagResult> Int32Res = Bag.GetValueInt32(TemperatureName);
+		AITEST_TRUE(TEXT("Bag Temperature as int32 should be 451"), Int32Res.IsValid() && Int32Res.GetValue() == 451);
+
+		return true;
+	}
+};
+IMPLEMENT_AI_INSTANT_TEST(FTest_MigrateProperty, "System.StructUtils.PropertyBag.MigrateProperty");
+
+struct FTest_Object : FAITestBase
+{
+	virtual bool InstantTest() override
+	{
+		static const FName ObjectName(TEXT("Object"));
+
+		UBagTestObject1* Test1 = NewObject<UBagTestObject1>(&GetWorld());
+		UBagTestObject2* Test2 = NewObject<UBagTestObject2>(&GetWorld());
+		UBagTestObject1Derived* Test1Derived = NewObject<UBagTestObject1Derived>(&GetWorld());
+
+		FInstancedPropertyBag Bag;
+		Bag.AddProperty(ObjectName, EPropertyBagPropertyType::Object, UBagTestObject1::StaticClass());
+		AITEST_TRUE(TEXT("Bag should have Object property"), Bag.FindPropertyDescByName(ObjectName) != nullptr);
+
+		AITEST_TRUE(TEXT("Bag set Object to Test1Derived should succeed"), Bag.SetValueObject(ObjectName, Test1Derived) == EPropertyBagResult::Success);
+		AITEST_FALSE(TEXT("Bag set Object to Test2 should fail"), Bag.SetValueObject(ObjectName, Test2) == EPropertyBagResult::Success);
+		AITEST_TRUE(TEXT("Bag set Object to Test1 should succeed"), Bag.SetValueObject(ObjectName, Test1) == EPropertyBagResult::Success);
+
+		TValueOrError<UBagTestObject1*, EPropertyBagResult> Test1Res = Bag.GetValueObject<UBagTestObject1>(ObjectName);
+		TValueOrError<UBagTestObject1Derived*, EPropertyBagResult> Test1DerivedRes = Bag.GetValueObject<UBagTestObject1Derived>(ObjectName);
+		
+		AITEST_TRUE(TEXT("Bag get Object as Test1 should succeed"), Test1Res.IsValid());
+		AITEST_FALSE(TEXT("Bag get Object as Test1Derived should fail"), Test1DerivedRes.IsValid()); // Note: the current value is Test1, and Cast should fail.
+
+		// Test conversion from Object to SoftObject
+		Bag.AddProperty(ObjectName, EPropertyBagPropertyType::SoftObject, UBagTestObject1::StaticClass());
+		TValueOrError<UBagTestObject1*, EPropertyBagResult> Test1Res2 = Bag.GetValueObject<UBagTestObject1>(ObjectName);
+		AITEST_TRUE(TEXT("Bag get Object as Test1 should succeed after migration soft object"), Test1Res2.IsValid());
+		AITEST_TRUE(TEXT("Bag get Object Test1 should be Test1 after migration soft object"), Test1Res2.GetValue() == Test1);
+
+		// Test conversion from SoftObject to Object
+		Bag.AddProperty(ObjectName, EPropertyBagPropertyType::Object, UBagTestObject1::StaticClass());
+		TValueOrError<UBagTestObject1*, EPropertyBagResult> Test1Res3 = Bag.GetValueObject<UBagTestObject1>(ObjectName);
+		AITEST_TRUE(TEXT("Bag get Object as Test1 should succeed after migration object"), Test1Res3.IsValid());
+		AITEST_TRUE(TEXT("Bag get Object Test1 should be Test1 after migration object"), Test1Res3.GetValue() == Test1);
+
+		// Test conversion from different type
+		Bag.AddProperty(ObjectName, EPropertyBagPropertyType::Object, UBagTestObject2::StaticClass());
+		TValueOrError<UBagTestObject1*, EPropertyBagResult> Test1Res4 = Bag.GetValueObject<UBagTestObject1>(ObjectName);
+		TValueOrError<UBagTestObject2*, EPropertyBagResult> Test2Res = Bag.GetValueObject<UBagTestObject2>(ObjectName);
+		AITEST_FALSE(TEXT("Bag get Object as Test1 should fail after migration to test2"), Test1Res4.IsValid());
+		AITEST_TRUE(TEXT("Bag get Object as Test1 should succeed after migration to test2"), Test2Res.IsValid());
+		AITEST_TRUE(TEXT("Bag get Object Test2 should be null after migration to test2"), Test2Res.GetValue() == nullptr);
+
+		return true;
+	}
+};
+IMPLEMENT_AI_INSTANT_TEST(FTest_Object, "System.StructUtils.PropertyBag.Object");
+
+struct FTest_Struct : FAITestBase
+{
+	virtual bool InstantTest() override
+	{
+		static const FName StructName(TEXT("Struct"));
+
+		FInstancedPropertyBag Bag;
+		Bag.AddProperty(StructName, EPropertyBagPropertyType::Struct, FTestStructSimple::StaticStruct());
+		AITEST_TRUE(TEXT("Bag should have Struct property"), Bag.FindPropertyDescByName(StructName) != nullptr);
+
+		FTestStructSimple Value;
+		Value.Float = 42.0f;
+
+		FTestStructComplex Value2;
+
+		AITEST_TRUE(TEXT("Bag set Struct as struct view should succeed"), Bag.SetValueStruct(StructName, FConstStructView::Make(Value)) == EPropertyBagResult::Success);
+		AITEST_TRUE(TEXT("Bag set Struct with template should succeed"), Bag.SetValueStruct(StructName, Value) == EPropertyBagResult::Success);
+		AITEST_FALSE(TEXT("Bag set Struct to complex as struct view should succeed"), Bag.SetValueStruct(StructName, FConstStructView::Make(Value2)) == EPropertyBagResult::Success);
+		AITEST_FALSE(TEXT("Bag set Struct to complex  with template should succeed"), Bag.SetValueStruct(StructName, Value2) == EPropertyBagResult::Success);
+		
+		TValueOrError<FStructView, EPropertyBagResult> Res1 = Bag.GetValueStruct(StructName);
+		TValueOrError<FTestStructSimple*, EPropertyBagResult> Res2 = Bag.GetValueStruct<FTestStructSimple>(StructName);
+		TValueOrError<FTestStructSimpleBase*, EPropertyBagResult> Res3 = Bag.GetValueStruct<FTestStructSimpleBase>(StructName);
+		TValueOrError<FTestStructComplex*, EPropertyBagResult> Res4 = Bag.GetValueStruct<FTestStructComplex>(StructName);
+		
+		AITEST_TRUE(TEXT("Bag get Struct as struct view should succeed"), Res1.IsValid());
+		AITEST_TRUE(TEXT("Bag get Struct as simple should succeed"), Res2.IsValid());
+		AITEST_TRUE(TEXT("Bag result value should be 42"), FMath::IsNearlyEqual(Res2.GetValue()->Float, 42.0f));
+		AITEST_TRUE(TEXT("Bag get Struct as simple base should succeed"), Res3.IsValid());
+		AITEST_FALSE(TEXT("Bag get Struct as complex should succeed"), Res4.IsValid());
+
+		Bag.AddProperty(StructName, EPropertyBagPropertyType::Bool);
+		TValueOrError<FStructView, EPropertyBagResult> MigRes1 = Bag.GetValueStruct(StructName);
+		TValueOrError<FTestStructSimple*, EPropertyBagResult> MigRes2 = Bag.GetValueStruct<FTestStructSimple>(StructName);
+		
+		AITEST_FALSE(TEXT("Bag get Struct as struct view should fail after migration"), MigRes1.IsValid());
+		AITEST_FALSE(TEXT("Bag get Struct as simple should succeed after migration"), MigRes2.IsValid());
+
+		return true;
+	}
+};
+IMPLEMENT_AI_INSTANT_TEST(FTest_Struct, "System.StructUtils.PropertyBag.Struct");
+
+struct FTest_Class : FAITestBase
+{
+	virtual bool InstantTest() override
+	{
+		static const FName ClassName(TEXT("Class"));
+
+		FInstancedPropertyBag Bag;
+		Bag.AddProperty(ClassName, EPropertyBagPropertyType::Class, UBagTestObject1::StaticClass());
+		AITEST_TRUE(TEXT("Bag should have Class property"), Bag.FindPropertyDescByName(ClassName) != nullptr);
+
+		AITEST_TRUE(TEXT("Bag set Class to UBagTestObject1 should succeed"), Bag.SetValueClass(ClassName, UBagTestObject1::StaticClass()) == EPropertyBagResult::Success);
+		AITEST_FALSE(TEXT("Bag set Class to UBagTestObject2 should fail"), Bag.SetValueClass(ClassName, UBagTestObject2::StaticClass()) == EPropertyBagResult::Success);
+		AITEST_TRUE(TEXT("Bag set Class to UBagTestObject1Derived should succeed"), Bag.SetValueClass(ClassName, UBagTestObject1Derived::StaticClass()) == EPropertyBagResult::Success);
+
+		TValueOrError<UClass*, EPropertyBagResult> Res1 = Bag.GetValueClass(ClassName);
+		AITEST_TRUE(TEXT("Bag get Class should succeed"), Res1.IsValid());
+		AITEST_TRUE(TEXT("Bag Class result should be UBagTestObject1Derived"), Res1.GetValue() == UBagTestObject1Derived::StaticClass());
+
+		return true;
+	}
+};
+IMPLEMENT_AI_INSTANT_TEST(FTest_Class, "System.StructUtils.PropertyBag.Class");
+
+struct FTest_Enum : FAITestBase
+{
+	virtual bool InstantTest() override
+	{
+		static const FName EnumName(TEXT("Enum"));
+
+		FInstancedPropertyBag Bag;
+		Bag.AddProperty(EnumName, EPropertyBagPropertyType::Enum, StaticEnum<EPropertyBagTest1>());
+		AITEST_TRUE(TEXT("Bag should have Enum property"), Bag.FindPropertyDescByName(EnumName) != nullptr);
+
+		AITEST_TRUE(TEXT("Bag set Enum to Foo should succeed"), Bag.SetValueEnum(EnumName, EPropertyBagTest1::Foo) == EPropertyBagResult::Success);
+		AITEST_FALSE(TEXT("Bag set Enum to Bongo should fail"), Bag.SetValueEnum(EnumName, EPropertyBagTest2::Bongo) == EPropertyBagResult::Success);
+		
+		TValueOrError<EPropertyBagTest1, EPropertyBagResult> Res1 = Bag.GetValueEnum<EPropertyBagTest1>(EnumName);
+		TValueOrError<EPropertyBagTest2, EPropertyBagResult> Res2 = Bag.GetValueEnum<EPropertyBagTest2>(EnumName);
+		
+		AITEST_TRUE(TEXT("Bag get Enum should succeed"), Res1.IsValid());
+		AITEST_TRUE(TEXT("Bag Enum result should be Byte"), Res1.GetValue() == EPropertyBagTest1::Foo);
+		AITEST_FALSE(TEXT("Bag get Enum with different type should fail"), Res2.IsValid());
+
+		return true;
+	}
+};
+IMPLEMENT_AI_INSTANT_TEST(FTest_Enum, "System.StructUtils.PropertyBag.Enum");
+
+
+} // FPropertyBagTest
+
+PRAGMA_ENABLE_OPTIMIZATION
+
+#undef LOCTEXT_NAMESPACE
diff --git a/Engine/Plugins/Experimental/StructUtils/Source/StructUtilsTestSuite/Public/StructUtilsTestTypes.h b/Engine/Plugins/Experimental/StructUtils/Source/StructUtilsTestSuite/Public/StructUtilsTestTypes.h
--- a/Engine/Plugins/Experimental/StructUtils/Source/StructUtilsTestSuite/Public/StructUtilsTestTypes.h
+++ b/Engine/Plugins/Experimental/StructUtils/Source/StructUtilsTestSuite/Public/StructUtilsTestTypes.h
@@ -114,3 +114,35 @@ struct FTestStructSimple7 : public FTestStructSimpleBase
 	UPROPERTY()
 	float Float = 0.0f;
 };
+
+UCLASS()
+class UBagTestObject1 : public UObject
+{
+	GENERATED_BODY()
+};
+
+UCLASS()
+class UBagTestObject2 : public UObject
+{
+	GENERATED_BODY()
+};
+
+UCLASS()
+class UBagTestObject1Derived : public UBagTestObject1
+{
+	GENERATED_BODY()
+};
+
+UENUM()
+enum class EPropertyBagTest1 : uint8
+{
+	Foo,
+	Bar,
+};
+
+UENUM()
+enum class EPropertyBagTest2 : uint8
+{
+	Bingo,
+	Bongo,
+};
