Renamed TScriptStructTypeBitSet to TStructTypeBitSet and extended it to handle UStructs as well.

diff --git a/Engine/Plugins/Experimental/StructUtils/Source/StructUtils/Public/ScriptStructTypeBitSet.h b/Engine/Plugins/Experimental/StructUtils/Source/StructUtils/Public/ScriptStructTypeBitSet.h
deleted file mode 100644
diff --git a/Engine/Plugins/Experimental/StructUtils/Source/StructUtils/Public/StructTypeBitSet.h b/Engine/Plugins/Experimental/StructUtils/Source/StructUtils/Public/StructTypeBitSet.h
new file mode 100644
--- /dev/null
+++ b/Engine/Plugins/Experimental/StructUtils/Source/StructUtils/Public/StructTypeBitSet.h
@@ -0,0 +1,578 @@
+// Copyright Epic Games, Inc. All Rights Reserved.
+
+#pragma once
+
+#include "UObject/Class.h"
+#include "Containers/BitArray.h"
+#include "InstancedStruct.h"
+#include "StructUtilsTypes.h"
+
+
+/**
+ * The TStructTypeBitSet holds information on "existence" of subtypes of a given UStruct. The information on 
+ * available child-structs is gathered lazily - the internal FStructTracker is assigning a given type a new index the
+ * very first time the type is encountered. 
+ * To create a specific instantiation of the type you need to declare the static StructTracker member variable. 
+ * To do it for an arbitrary type FFooBar add the following in your header or cpp file:
+ * 
+ *	DECLARE_STRUCTTYPEBITSET(FMyFooBarBitSet, FFooBar);
+ * 
+ * where FMyFooBarBitSet is the alias of the type you can use in your code. To have your type exposed to other modules 
+ * use DECLARE_STRUCTTYPEBITSET_EXPORTED, like so:
+ * 
+ *	DECLARE_STRUCTTYPEBITSET_EXPORTED(MYMODULE_API, FMyFooBarBitSet, FFooBar);
+ *
+ * The type contains static members so you'll also need to define these. You can easily do it by placing the following 
+ * in your cpp file (continuing the FFooBar example):
+ *
+ *	DEFINE_STRUCTTYPEBITSET(FMyFooBarBitSet);
+ * 
+ */
+struct FStructTracker
+{
+	int32 FindOrAddStructTypeIndex(const UStruct& InStructType)
+	{
+		// Get existing index...
+		const uint32 Hash = PointerHash(&InStructType);
+		FSetElementId ElementId = StructTypeToIndexSet.FindIdByHash(Hash, Hash);
+
+		if (!ElementId.IsValidId())
+		{
+			// .. or create new one
+			ElementId = StructTypeToIndexSet.AddByHash(Hash, Hash);
+			checkSlow(ElementId.IsValidId());
+
+			checkSlow(StructTypesList.Num() == ElementId.AsInteger());
+			StructTypesList.Add(&InStructType);
+		}
+		const int32 Index = ElementId.AsInteger();
+
+#if WITH_STRUCTUTILS_DEBUG
+		if (Index == DebugStructTypeNamesList.Num())
+		{
+			DebugStructTypeNamesList.Add(InStructType.GetFName());
+			ensure(StructTypeToIndexSet.Num() == DebugStructTypeNamesList.Num());
+		}
+#endif // WITH_STRUCTUTILS_DEBUG
+		return Index;
+	}
+
+	const UStruct* GetStructType(const int32 StructTypeIndex) const
+	{
+		return StructTypesList.IsValidIndex(StructTypeIndex) ? StructTypesList[StructTypeIndex].Get() : nullptr;
+	}
+
+#if WITH_STRUCTUTILS_DEBUG
+	/**
+	* @return index identifying given tag or INDEX_NONE if it has never been used/seen before.
+	*/
+	FName DebugGetStructTypeName(const int32 StructTypeIndex) const
+	{
+		return DebugStructTypeNamesList.IsValidIndex(StructTypeIndex) ? DebugStructTypeNamesList[StructTypeIndex] : FName();
+	}
+
+	TConstArrayView<TWeakObjectPtr<const UStruct>> DebugGetAllStructTypes() const { return StructTypesList; }
+
+	void DebugResetStructTypeMappingInfo()
+	{
+		StructTypeToIndexSet.Reset();
+		StructTypesList.Reset();
+		DebugStructTypeNamesList.Reset();
+	}
+	TArray<FName, TInlineAllocator<64>> DebugStructTypeNamesList;
+#endif // WITH_STRUCTUTILS_DEBUG
+
+	TSet<uint32> StructTypeToIndexSet;
+	TArray<TWeakObjectPtr<const UStruct>, TInlineAllocator<64>> StructTypesList;
+};
+
+template<typename TBaseStruct>
+struct TStructTypeBitSet
+{
+private:
+	struct FBitArrayExt : TBitArray<>
+	{
+		FBitArrayExt() = default;
+		FBitArrayExt(const TBitArray<>& Source) : TBitArray<>(Source)
+		{}
+
+		FBitArrayExt& operator=(const TBitArray<>& Other)
+		{
+			*((TBitArray<>*)this) = Other;
+			return *this;
+		}
+
+		FORCEINLINE bool HasAll(const TBitArray<>& Other) const
+		{
+			FConstWordIterator ThisIterator(*this);
+			FConstWordIterator OtherIterator(Other);
+
+			while (ThisIterator || OtherIterator)
+			{
+				const uint32 A = ThisIterator ? ThisIterator.GetWord() : 0;
+				const uint32 B = OtherIterator ? OtherIterator.GetWord() : 0;
+				if ((A & B) != B)
+				{
+					return false;
+				}
+
+				++ThisIterator;
+				++OtherIterator;
+			}
+
+			return true;
+		}
+
+		FORCEINLINE bool HasAny(const TBitArray<>& Other) const
+		{
+			FConstWordIterator ThisIterator(*this);
+			FConstWordIterator OtherIterator(Other);
+
+			while (ThisIterator || OtherIterator)
+			{
+				const uint32 A = ThisIterator ? ThisIterator.GetWord() : 0;
+				const uint32 B = OtherIterator ? OtherIterator.GetWord() : 0;
+				if ((A & B) != 0)
+				{
+					return true;
+				}
+
+				++ThisIterator;
+				++OtherIterator;
+			}
+
+			return false;
+		}
+
+		FORCEINLINE bool IsEmpty() const
+		{
+			FConstWordIterator Iterator(*this);
+
+			while (Iterator && Iterator.GetWord() == 0)
+			{
+				++Iterator;
+			}
+
+			return !Iterator;
+		}
+
+		FORCEINLINE void operator-=(const TBitArray<>& Other)
+		{
+			FWordIterator ThisIterator(*this);
+			FConstWordIterator OtherIterator(Other);
+
+			while (ThisIterator && OtherIterator)
+			{
+				ThisIterator.SetWord(ThisIterator.GetWord() & ~OtherIterator.GetWord());
+
+				++ThisIterator;
+				++OtherIterator;
+			}
+		}
+
+		FORCEINLINE friend uint32 GetTypeHash(const FBitArrayExt& Instance)
+		{
+			FConstWordIterator Iterator(Instance);
+			uint32 Hash = 0;
+			uint32 TrailingZeroHash = 0;
+			while (Iterator)
+			{
+				const uint32 Word = Iterator.GetWord();
+				if (Word)
+				{
+					Hash = HashCombine(TrailingZeroHash ? TrailingZeroHash : Hash, Word);
+					TrailingZeroHash = 0;
+				}
+				else // potentially a trailing 0-word
+				{
+					TrailingZeroHash = HashCombine(TrailingZeroHash ? TrailingZeroHash : Hash, Word);
+				}
+				++Iterator;
+			}
+			return Hash;
+		}
+
+		void AddAtIndex(const int32 Index)
+		{
+			PadToNum(Index + 1, false);
+			SetBitNoCheck(Index, true);
+		}
+
+		void RemoveAtIndex(const int32 Index)
+		{
+			check(Index >= 0);
+			if (Index < Num())
+			{
+				SetBitNoCheck(Index, false);
+			}
+			// else, it's already not present
+		}
+
+		bool Contains(const int32 Index) const
+		{
+			check(Index >= 0);
+			return (Index < Num()) && (*this)[Index];
+		}
+
+	protected:
+		/**
+		 * duplication of TBitArray::SetBitNoCheck needed since it's private but it's the performant way of setting bits
+		 * when we know the index is valid.
+		 * @todo ask Core team about exposing that
+		 */
+		void SetBitNoCheck(const int32 Index, const bool Value)
+		{
+			uint32& Word = GetData()[Index / NumBitsPerDWORD];
+			const uint32 BitOffset = (Index % NumBitsPerDWORD);
+			Word = (Word & ~(1 << BitOffset)) | (((uint32)Value) << BitOffset);
+		}
+	};
+
+	static FStructTracker StructTracker;
+	using FBaseStruct = TBaseStruct;
+
+public:
+	TStructTypeBitSet() = default;
+
+	explicit TStructTypeBitSet(const UStruct& StructType)
+	{
+		Add(StructType);
+	}
+
+	explicit TStructTypeBitSet(std::initializer_list<const UStruct*> InitList)
+	{
+		for (const UStruct* StructType : InitList)
+		{
+			if (StructType)
+			{
+				Add(*StructType);
+			}
+		}
+	}
+
+	explicit TStructTypeBitSet(TConstArrayView<const UStruct*> InitList)
+	{
+		for (const UStruct* StructType : InitList)
+		{
+			if (StructType)
+			{
+				Add(*StructType);
+			}
+		}
+	}
+
+	explicit TStructTypeBitSet(TConstArrayView<const UScriptStruct*> InitList)
+	{
+		for (const UStruct* StructType : InitList)
+		{
+			if (StructType)
+			{
+				Add(*StructType);
+			}
+		}
+	}
+
+	explicit TStructTypeBitSet(TConstArrayView<FInstancedStruct> InitList)
+	{
+		for (const FInstancedStruct& StructInstance : InitList)
+		{
+			if (StructInstance.GetScriptStruct())
+			{
+				Add(*StructInstance.GetScriptStruct());
+			}
+		}
+	}
+
+private:
+	/** 
+	 * A private constructor for a creating an instance straight from TBitArrays. 
+	 * @Note that this constructor needs to remain private to ensure consistency of stored values with data tracked 
+	 * by the StructTracker
+	 */
+	TStructTypeBitSet(const TBitArray<>& Source)
+		: StructTypesBitArray(Source)
+	{
+	}
+
+	TStructTypeBitSet(const int32 BitToSet)
+	{
+		StructTypesBitArray.AddAtIndex(BitToSet);
+	}
+
+	static const UStruct* GetBaseUStruct()
+	{
+		static const UStruct* Instance = UE::StructUtils::GetAsUStruct<TBaseStruct>();
+		return Instance;
+	}
+
+public:
+
+	static int32 CreateTypeIndex(const UStruct& InStructType)
+	{
+#if WITH_STRUCTUTILS_DEBUG
+		ensureMsgf(InStructType.IsChildOf(GetBaseUStruct())
+			, TEXT("Creating index for '%s' while it doesn't derive from the expected struct type %s")
+			, *InStructType.GetPathName(), *GetBaseUStruct()->GetName());
+#endif // WITH_STRUCTUTILS_DEBUG
+
+		return StructTracker.FindOrAddStructTypeIndex(InStructType);
+	}
+
+	template<typename T>
+	static int32 GetTypeIndex()
+	{
+		static_assert(TIsDerivedFrom<T, TBaseStruct>::IsDerived, "Given struct type doesn't match the expected base struct type.");
+		static const int32 TypeIndex = CreateTypeIndex(*UE::StructUtils::GetAsUStruct<T>());
+		return TypeIndex;
+	}
+
+	template<typename T>
+	static const TStructTypeBitSet<TBaseStruct>& GetTypeBitSet()
+	{
+		static_assert(TIsDerivedFrom<T, TBaseStruct>::IsDerived, "Given struct type doesn't match the expected base struct type.");
+		static const TStructTypeBitSet<TBaseStruct> TypeBitSet(GetTypeIndex<T>());
+		return TypeBitSet;
+	}
+	
+	template<typename T>
+	FORCEINLINE void Add()
+	{
+		static_assert(TIsDerivedFrom<T, TBaseStruct>::IsDerived, "Given struct type doesn't match the expected base struct type.");
+		const int32 StructTypeIndex = GetTypeIndex<T>();
+		StructTypesBitArray.AddAtIndex(StructTypeIndex);
+	}
+
+	template<typename T>
+	FORCEINLINE void Remove()
+	{
+		static_assert(TIsDerivedFrom<T, TBaseStruct>::IsDerived, "Given struct type doesn't match the expected base struct type.");
+		const int32 StructTypeIndex = GetTypeIndex<T>();
+		StructTypesBitArray.RemoveAtIndex(StructTypeIndex);
+	}
+
+	FORCEINLINE void Remove(const TStructTypeBitSet<TBaseStruct>& Other)
+	{
+		StructTypesBitArray -= Other.StructTypesBitArray;
+	}
+
+	template<typename T>
+	FORCEINLINE bool Contains() const
+	{
+		static_assert(TIsDerivedFrom<T, TBaseStruct>::IsDerived, "Given struct type doesn't match the expected base struct type.");
+		const int32 StructTypeIndex = GetTypeIndex<T>();
+		return StructTypesBitArray.Contains(StructTypeIndex);
+	}
+
+	void Add(const UStruct& InStructType)
+	{
+#if WITH_STRUCTUTILS_DEBUG
+		ensureMsgf(InStructType.IsChildOf(GetBaseUStruct())
+				, TEXT("Registering '%s' with FStructTracker while it doesn't derive from the expected struct type %s")
+				, *InStructType.GetPathName(), *GetBaseUStruct()->GetName());
+#endif // WITH_STRUCTUTILS_DEBUG
+
+		const int32 StructTypeIndex = StructTracker.FindOrAddStructTypeIndex(InStructType);
+		StructTypesBitArray.AddAtIndex(StructTypeIndex);
+	}
+
+	void Remove(const UStruct& InStructType)
+	{
+#if WITH_STRUCTUTILS_DEBUG
+		ensureMsgf(InStructType.IsChildOf(GetBaseUStruct())
+				, TEXT("Registering '%s' with FStructTracker while it doesn't derive from the expected struct type %s")
+				, *InStructType.GetPathName(), *GetBaseUStruct()->GetName());
+#endif // WITH_STRUCTUTILS_DEBUG
+
+		const int32 StructTypeIndex = StructTracker.FindOrAddStructTypeIndex(InStructType);
+		StructTypesBitArray.RemoveAtIndex(StructTypeIndex);
+	}
+
+	void Reset() { StructTypesBitArray.Reset(); }
+
+	bool Contains(const UStruct& InStructType) const
+	{
+#if WITH_STRUCTUTILS_DEBUG
+		ensureMsgf(InStructType.IsChildOf(GetBaseUStruct())
+				, TEXT("Registering '%s' with FStructTracker while it doesn't derive from the expected struct type %s")
+				, *InStructType.GetPathName(), *GetBaseUStruct()->GetName());
+#endif // WITH_STRUCTUTILS_DEBUG
+
+		const int32 StructTypeIndex = StructTracker.FindOrAddStructTypeIndex(InStructType);
+		return StructTypesBitArray.Contains(StructTypeIndex);
+	}
+
+	FORCEINLINE TStructTypeBitSet operator+(const TStructTypeBitSet& Other) const
+	{
+		TStructTypeBitSet Result;
+		Result.StructTypesBitArray = TBitArray<>::BitwiseOR(StructTypesBitArray, Other.StructTypesBitArray, EBitwiseOperatorFlags::MaxSize);
+		return MoveTemp(Result);
+	}
+
+	FORCEINLINE void operator+=(const TStructTypeBitSet& Other)
+	{
+		StructTypesBitArray = TBitArray<>::BitwiseOR(StructTypesBitArray, Other.StructTypesBitArray, EBitwiseOperatorFlags::MaxSize);
+	}
+
+	FORCEINLINE void operator-=(const TStructTypeBitSet& Other)
+	{
+		StructTypesBitArray -= Other.StructTypesBitArray;
+	}
+
+	FORCEINLINE TStructTypeBitSet operator+(const UStruct& NewElement) const
+	{
+		TStructTypeBitSet Result = *this;
+		Result.Add(NewElement);
+		return MoveTemp(Result);
+	}
+
+	FORCEINLINE TStructTypeBitSet operator-(const UStruct& NewElement) const
+	{
+		TStructTypeBitSet Result = *this;
+		Result.Remove(NewElement);
+		return MoveTemp(Result);
+	}
+
+	FORCEINLINE TStructTypeBitSet operator-(const TStructTypeBitSet& Other) const
+	{
+		TStructTypeBitSet Result = *this;
+		Result -= Other;
+		return MoveTemp(Result);
+	}
+
+	FORCEINLINE TStructTypeBitSet operator&(const TStructTypeBitSet& Other) const
+	{
+		return TStructTypeBitSet(TBitArray<>::BitwiseAND(StructTypesBitArray, Other.StructTypesBitArray, EBitwiseOperatorFlags::MinSize));
+	}
+
+	FORCEINLINE TStructTypeBitSet operator|(const TStructTypeBitSet& Other) const
+	{
+		return TStructTypeBitSet(TBitArray<>::BitwiseOR(StructTypesBitArray, Other.StructTypesBitArray, EBitwiseOperatorFlags::MaxSize));
+	}
+
+	FORCEINLINE TStructTypeBitSet GetOverlap(const TStructTypeBitSet& Other) const
+	{
+		return *this & Other;
+	}
+
+	FORCEINLINE bool IsEquivalent(const TStructTypeBitSet<TBaseStruct>& Other) const
+	{
+		return StructTypesBitArray.CompareSetBits(Other.StructTypesBitArray, /*bMissingBitValue=*/false);
+	}
+
+	FORCEINLINE bool HasAll(const TStructTypeBitSet& Other) const
+	{
+		return StructTypesBitArray.HasAll(Other.StructTypesBitArray);
+	}
+
+	FORCEINLINE bool HasAny(const TStructTypeBitSet& Other) const
+	{
+		return StructTypesBitArray.HasAny(Other.StructTypesBitArray);
+	}
+
+	FORCEINLINE bool HasNone(const TStructTypeBitSet& Other) const
+	{
+		return !StructTypesBitArray.HasAny(Other.StructTypesBitArray);
+	}
+
+	bool IsEmpty() const 
+	{ 
+		return StructTypesBitArray.IsEmpty();
+	}
+
+	FORCEINLINE bool IsBitSet(const int32 BitIndex) const
+	{
+		return StructTypesBitArray.Contains(BitIndex);
+	}
+
+	FORCEINLINE bool operator==(const TStructTypeBitSet& Other) const { return StructTypesBitArray == Other.StructTypesBitArray; }
+	FORCEINLINE bool operator!=(const TStructTypeBitSet& Other) const { return !(*this == Other); }
+
+	/**
+	 * note that this function is slow(ish) due to the FStructTracker utilizing WeakObjectPtrs to store types. 
+	 * @todo To be improved.
+	 */
+	template<typename TOutStructType, typename Allocator>
+	void ExportTypes(TArray<const TOutStructType*, Allocator>& OutTypes) const
+	{
+		TBitArray<>::FConstIterator It(StructTypesBitArray);
+		while (It)
+		{
+			if (It.GetValue())
+			{
+				OutTypes.Add(Cast<TOutStructType>(StructTracker.GetStructType(It.GetIndex())));
+			}
+			++It;
+		}
+	}
+
+	SIZE_T GetAllocatedSize() const
+	{
+		return StructTypesBitArray.GetAllocatedSize();
+	}
+
+	FString DebugGetStringDesc() const
+	{
+#if WITH_STRUCTUTILS_DEBUG
+		FStringOutputDevice Ar;
+		DebugGetStringDesc(Ar);
+		return static_cast<FString>(Ar);
+#else
+		return TEXT("DEBUG INFO COMPILED OUT");
+#endif //WITH_STRUCTUTILS_DEBUG
+	}
+
+#if WITH_STRUCTUTILS_DEBUG
+	void DebugGetStringDesc(FOutputDevice& Ar) const
+	{
+		for (int32 Index = 0; Index < StructTypesBitArray.Num(); ++Index)
+		{
+			if (StructTypesBitArray[Index])
+			{
+				Ar.Logf(TEXT("%s, "), *StructTracker.DebugGetStructTypeName(Index).ToString());
+			}
+		}
+	}
+
+	void DebugGetIndividualNames(TArray<FName>& OutFNames) const
+	{
+		for (int32 Index = 0; Index < StructTypesBitArray.Num(); ++Index)
+		{
+			if (StructTypesBitArray[Index])
+			{
+				OutFNames.Add(StructTracker.DebugGetStructTypeName(Index));
+			}
+		}
+	}
+
+	static TConstArrayView<TWeakObjectPtr<const UStruct>> DebugGetAllStructTypes()
+	{
+		return StructTracker.DebugGetAllStructTypes();
+	}
+
+	/**
+	 * Resets all the information gathered on the tags. Calling this results in invalidating all previously created
+	 * FStructTypeCollection instances. Used only for debugging and unit/functional testing.
+	 */
+	static void DebugResetStructTypeMappingInfo()
+	{
+		StructTracker.DebugResetStructTypeMappingInfo();
+	}
+protected:
+	// unittesting purposes only
+	const TBitArray<>& DebugGetStructTypesBitArray() const { return StructTypesBitArray; }
+	TBitArray<>& DebugGetMutableStructTypesBitArray() { return StructTypesBitArray; }
+#endif // WITH_STRUCTUTILS_DEBUG
+
+public:
+	FORCEINLINE friend uint32 GetTypeHash(const TStructTypeBitSet<TBaseStruct>& Instance)
+	{
+		const uint32 StoredTypeHash = PointerHash(UE::StructUtils::GetAsUStruct<TBaseStruct>());
+		const uint32 BitArrayHash = GetTypeHash(Instance.StructTypesBitArray);
+		return HashCombine(StoredTypeHash, BitArrayHash);
+	}
+
+private:
+	FBitArrayExt StructTypesBitArray;
+};
+
+template<typename TBaseStruct> FStructTracker TStructTypeBitSet<TBaseStruct>::StructTracker;
diff --git a/Engine/Plugins/Experimental/StructUtils/Source/StructUtils/Public/StructUtilsTypes.h b/Engine/Plugins/Experimental/StructUtils/Source/StructUtils/Public/StructUtilsTypes.h
--- a/Engine/Plugins/Experimental/StructUtils/Source/StructUtils/Public/StructUtilsTypes.h
+++ b/Engine/Plugins/Experimental/StructUtils/Source/StructUtils/Public/StructUtilsTypes.h
@@ -4,6 +4,7 @@
 
 #include "UObject/Class.h"
 #include "StructView.h"
+#include "Templates/UnrealTypeTraits.h"
 
 #ifndef WITH_STRUCTUTILS_DEBUG
 #define WITH_STRUCTUTILS_DEBUG (!(UE_BUILD_SHIPPING || UE_BUILD_SHIPPING_WITH_EDITOR || UE_BUILD_TEST) && 1)
@@ -17,6 +18,18 @@ namespace UE::StructUtils
 	extern STRUCTUTILS_API uint32 GetStructCrc32(const UScriptStruct& ScriptStruct, const uint8* StructMemory, const uint32 CRC = 0);
 
 	extern STRUCTUTILS_API uint32 GetStructCrc32(const FConstStructView StructView, const uint32 CRC = 0);
+
+	template <typename T>
+	typename TEnableIf<!TIsDerivedFrom<T, UObject>::IsDerived, UScriptStruct*>::Type GetAsUStruct()
+	{
+		return T::StaticStruct();
+	}
+
+	template <typename T>
+	typename TEnableIf<TIsDerivedFrom<T, UObject>::IsDerived, UClass*>::Type GetAsUStruct()
+	{
+		return T::StaticClass();
+	}
 }
 
 /* Predicate useful to find a struct of a specific type in an container */
diff --git a/Engine/Plugins/Experimental/StructUtils/Source/StructUtilsTestSuite/Private/ScriptStructTypeBitSetTest.cpp b/Engine/Plugins/Experimental/StructUtils/Source/StructUtilsTestSuite/Private/ScriptStructTypeBitSetTest.cpp
--- a/Engine/Plugins/Experimental/StructUtils/Source/StructUtilsTestSuite/Private/ScriptStructTypeBitSetTest.cpp
+++ b/Engine/Plugins/Experimental/StructUtils/Source/StructUtilsTestSuite/Private/ScriptStructTypeBitSetTest.cpp
@@ -2,7 +2,7 @@
 
 #include "CoreMinimal.h"
 #include "AITestsCommon.h"
-#include "ScriptStructTypeBitSet.h"
+#include "StructTypeBitSet.h"
 #include "StructUtilsTestTypes.h"
 
 #if WITH_STRUCTUTILS_DEBUG
@@ -11,16 +11,16 @@
 
 PRAGMA_DISABLE_OPTIMIZATION
 
-namespace FScriptStructTypeBitSetTests
+namespace FStructTypeBitSetTests
 {
 
-using FTestStructBaseBitSet = TScriptStructTypeBitSet<FTestStructSimpleBase>;
+using FTestStructBaseBitSet = TStructTypeBitSet<FTestStructSimpleBase>;
 
-struct FTestStructBitSet : public TScriptStructTypeBitSet<FTestStructSimple>
+struct FTestStructBitSet : public TStructTypeBitSet<FTestStructSimple>
 {
 	FTestStructBitSet() = default;
-	FTestStructBitSet(const TScriptStructTypeBitSet<FTestStructSimple>& Other) 
-		: TScriptStructTypeBitSet<FTestStructSimple>(Other)
+	FTestStructBitSet(const TStructTypeBitSet<FTestStructSimple>& Other) 
+		: TStructTypeBitSet<FTestStructSimple>(Other)
 	{
 
 	}
@@ -320,7 +320,7 @@ struct FStructUtilsTest_CompileTimeBitSetReversed : FAITestBase
 
 IMPLEMENT_AI_INSTANT_TEST(FStructUtilsTest_CompileTimeBitSetReversed, "System.StructUtils.BitSet.CompileTime_Reversed");
 
-} // namespace FScriptStructTypeBitSetTests
+} // namespace FStructTypeBitSetTests
 
 #undef LOCTEXT_NAMESPACE
 
diff --git a/Engine/Plugins/Runtime/MassEntity/Source/MassEntity/Private/MassEntityDebug.cpp b/Engine/Plugins/Runtime/MassEntity/Source/MassEntity/Private/MassEntityDebug.cpp
--- a/Engine/Plugins/Runtime/MassEntity/Source/MassEntity/Private/MassEntityDebug.cpp
+++ b/Engine/Plugins/Runtime/MassEntity/Source/MassEntity/Private/MassEntityDebug.cpp
@@ -5,7 +5,7 @@
 #include "MassProcessor.h"
 #include "Engine/World.h"
 #include "Engine/Engine.h"
-#include "ScriptStructTypeBitSet.h"
+#include "StructTypeBitSet.h"
 
 DEFINE_ENUM_TO_STRING(EMassProcessingPhase);
 
@@ -135,9 +135,9 @@ FAutoConsoleCommandWithWorldArgsAndOutputDevice LogFragmentSizes(
 	TEXT("Logs all the fragment types being used along with their sizes."),
 	FConsoleCommandWithWorldArgsAndOutputDeviceDelegate::CreateLambda([](const TArray<FString>& Params, UWorld* World, FOutputDevice& Ar)
 		{
-			for (const TWeakObjectPtr<const UScriptStruct>& WeakStruct : FMassFragmentBitSet::DebugGetAllStructTypes())
+			for (const TWeakObjectPtr<const UStruct>& WeakStruct : FMassFragmentBitSet::DebugGetAllStructTypes())
 			{
-				if (const UScriptStruct* StructType = WeakStruct.Get())
+				if (const UStruct* StructType = WeakStruct.Get())
 				{
 					Ar.Logf(ELogVerbosity::Log, TEXT("%s, size: %d"), *StructType->GetName(), StructType->GetStructureSize());
 				}
@@ -164,9 +164,9 @@ FAutoConsoleCommandWithOutputDevice LogFragments(
 	TEXT("Logs all the known tags and fragments along with their \"index\" as stored via bitsets."),
 	FConsoleCommandWithOutputDeviceDelegate::CreateStatic([](FOutputDevice& OutputDevice)
 {
-	auto PrintKnownTypes = [&OutputDevice](TConstArrayView<TWeakObjectPtr<const UScriptStruct>> AllStructs) {
+	auto PrintKnownTypes = [&OutputDevice](TConstArrayView<TWeakObjectPtr<const UStruct>> AllStructs) {
 		int i = 0;
-		for (TWeakObjectPtr<const UScriptStruct> Struct : AllStructs)
+		for (TWeakObjectPtr<const UStruct> Struct : AllStructs)
 		{
 			if (Struct.IsValid())
 			{
diff --git a/Engine/Plugins/Runtime/MassEntity/Source/MassEntity/Private/MassEntitySubsystem.cpp b/Engine/Plugins/Runtime/MassEntity/Source/MassEntity/Private/MassEntitySubsystem.cpp
--- a/Engine/Plugins/Runtime/MassEntity/Source/MassEntity/Private/MassEntitySubsystem.cpp
+++ b/Engine/Plugins/Runtime/MassEntity/Source/MassEntity/Private/MassEntitySubsystem.cpp
@@ -1352,7 +1352,7 @@ void UMassEntitySubsystem::DebugGetArchetypeFragmentTypes(const FMassArchetypeHa
 {
 	if (Archetype.DataPtr.IsValid())
 	{
-		Archetype.DataPtr->GetCompositionDescriptor().Fragments.DebugGetStructTypes(InOutFragmentList);
+		Archetype.DataPtr->GetCompositionDescriptor().Fragments.ExportTypes(InOutFragmentList);
 	}
 }
 
diff --git a/Engine/Plugins/Runtime/MassEntity/Source/MassEntity/Public/MassEntityTypes.h b/Engine/Plugins/Runtime/MassEntity/Source/MassEntity/Public/MassEntityTypes.h
--- a/Engine/Plugins/Runtime/MassEntity/Source/MassEntity/Public/MassEntityTypes.h
+++ b/Engine/Plugins/Runtime/MassEntity/Source/MassEntity/Public/MassEntityTypes.h
@@ -2,7 +2,7 @@
 
 #pragma once
 
-#include "ScriptStructTypeBitSet.h"
+#include "StructTypeBitSet.h"
 #include "MassProcessingTypes.h"
 #include "StructArrayView.h"
 #include "MassEntityTypes.generated.h"
@@ -102,14 +102,14 @@ struct FMassEntityHandle
 	}
 };
 
-template struct MASSENTITY_API TScriptStructTypeBitSet<FMassFragment>;
-using FMassFragmentBitSet = TScriptStructTypeBitSet<FMassFragment>;
-template struct MASSENTITY_API TScriptStructTypeBitSet<FMassTag>;
-using FMassTagBitSet = TScriptStructTypeBitSet<FMassTag>;
-template struct MASSENTITY_API TScriptStructTypeBitSet<FMassChunkFragment>;
-using FMassChunkFragmentBitSet = TScriptStructTypeBitSet<FMassChunkFragment>;
-template struct MASSENTITY_API TScriptStructTypeBitSet<FMassSharedFragment>;
-using FMassSharedFragmentBitSet = TScriptStructTypeBitSet<FMassSharedFragment>;
+template struct MASSENTITY_API TStructTypeBitSet<FMassFragment>;
+using FMassFragmentBitSet = TStructTypeBitSet<FMassFragment>;
+template struct MASSENTITY_API TStructTypeBitSet<FMassTag>;
+using FMassTagBitSet = TStructTypeBitSet<FMassTag>;
+template struct MASSENTITY_API TStructTypeBitSet<FMassChunkFragment>;
+using FMassChunkFragmentBitSet = TStructTypeBitSet<FMassChunkFragment>;
+template struct MASSENTITY_API TStructTypeBitSet<FMassSharedFragment>;
+using FMassSharedFragmentBitSet = TStructTypeBitSet<FMassSharedFragment>;
 
 /** The type summarily describing a composition of an entity or an archetype. It contains information on both the
  *  fragments as well as tags */
@@ -412,7 +412,7 @@ namespace UE::Mass
 {
 	/**
 	 * A statically-typed list of of related types. Used mainly to differentiate type collections at compile-type as well as
-	 * efficiently produce TScriptStructTypeBitSet representing given collection.
+	 * efficiently produce TStructTypeBitSet representing given collection.
 	 */
 	template<typename T, typename... TOthers>
 	struct TMultiTypeList : TMultiTypeList<TOthers...>
@@ -425,10 +425,10 @@ namespace UE::Mass
 		};
 
 		template<typename TBaseStruct>
-		static void PopulateBitSet(TScriptStructTypeBitSet<TBaseStruct>& OutBitSet)
+		static void PopulateBitSet(TStructTypeBitSet<TBaseStruct>& OutBitSet)
 		{
 			Super::PopulateBitSet(OutBitSet);
-			OutBitSet += TScriptStructTypeBitSet<TBaseStruct>::template GetTypeBitSet<FType>();
+			OutBitSet += TStructTypeBitSet<TBaseStruct>::template GetTypeBitSet<FType>();
 		}
 	};
 		
@@ -443,9 +443,9 @@ namespace UE::Mass
 		};
 
 		template<typename TBaseStruct>
-		static void PopulateBitSet(TScriptStructTypeBitSet<TBaseStruct>& OutBitSet)
+		static void PopulateBitSet(TStructTypeBitSet<TBaseStruct>& OutBitSet)
 		{
-			OutBitSet += TScriptStructTypeBitSet<TBaseStruct>::template GetTypeBitSet<FType>();
+			OutBitSet += TStructTypeBitSet<TBaseStruct>::template GetTypeBitSet<FType>();
 		}
 	};
 
