PropertyBag: Added missing copyright header

diff --git a/Engine/Plugins/Experimental/StructUtils/Source/StructUtils/Private/PropertyBag.cpp b/Engine/Plugins/Experimental/StructUtils/Source/StructUtils/Private/PropertyBag.cpp
--- a/Engine/Plugins/Experimental/StructUtils/Source/StructUtils/Private/PropertyBag.cpp
+++ b/Engine/Plugins/Experimental/StructUtils/Source/StructUtils/Private/PropertyBag.cpp
@@ -1,4 +1,6 @@
-﻿#include "PropertyBag.h"
+﻿// Copyright Epic Games, Inc. All Rights Reserved.
+
+#include "PropertyBag.h"
 #include "StructView.h"
 #include "Hash/CityHash.h"
 
