Added a UStruct-type parameter to TStructBitSet to further improve its type-safety

diff --git a/Engine/Plugins/Experimental/StructUtils/Source/StructUtils/Public/StructTypeBitSet.h b/Engine/Plugins/Experimental/StructUtils/Source/StructUtils/Public/StructTypeBitSet.h
--- a/Engine/Plugins/Experimental/StructUtils/Source/StructUtils/Public/StructTypeBitSet.h
+++ b/Engine/Plugins/Experimental/StructUtils/Source/StructUtils/Public/StructTypeBitSet.h
@@ -71,7 +71,11 @@ struct FStructTracker
 		return DebugStructTypeNamesList.IsValidIndex(StructTypeIndex) ? DebugStructTypeNamesList[StructTypeIndex] : FName();
 	}
 
-	TConstArrayView<TWeakObjectPtr<const UStruct>> DebugGetAllStructTypes() const { return StructTypesList; }
+	template<typename T>
+	TConstArrayView<TWeakObjectPtr<const T>> DebugGetAllStructTypes() const 
+	{ 
+		return MakeArrayView(reinterpret_cast<const TWeakObjectPtr<const T>*>(StructTypesList.GetData()), StructTypesList.Num());
+	}
 
 	void DebugResetStructTypeMappingInfo()
 	{
@@ -86,7 +90,7 @@ struct FStructTracker
 	TArray<TWeakObjectPtr<const UStruct>, TInlineAllocator<64>> StructTypesList;
 };
 
-template<typename TBaseStruct>
+template<typename TBaseStruct, typename TUStructType = UScriptStruct>
 struct TStructTypeBitSet
 {
 private:
@@ -229,30 +233,18 @@ private:
 	};
 
 	static FStructTracker StructTracker;
-	using FBaseStruct = TBaseStruct;
 
 public:
 	TStructTypeBitSet() = default;
 
-	explicit TStructTypeBitSet(const UStruct& StructType)
+	explicit TStructTypeBitSet(const TUStructType& StructType)
 	{
 		Add(StructType);
 	}
 
-	explicit TStructTypeBitSet(std::initializer_list<const UStruct*> InitList)
-	{
-		for (const UStruct* StructType : InitList)
-		{
-			if (StructType)
-			{
-				Add(*StructType);
-			}
-		}
-	}
-
-	explicit TStructTypeBitSet(TConstArrayView<const UStruct*> InitList)
+	explicit TStructTypeBitSet(std::initializer_list<const TUStructType*> InitList)
 	{
-		for (const UStruct* StructType : InitList)
+		for (const TUStructType* StructType : InitList)
 		{
 			if (StructType)
 			{
@@ -261,9 +253,9 @@ public:
 		}
 	}
 
-	explicit TStructTypeBitSet(TConstArrayView<const UScriptStruct*> InitList)
+	explicit TStructTypeBitSet(TConstArrayView<const TUStructType*> InitList)
 	{
-		for (const UStruct* StructType : InitList)
+		for (const TUStructType* StructType : InitList)
 		{
 			if (StructType)
 			{
@@ -272,6 +264,8 @@ public:
 		}
 	}
 
+	/** This flavor of constructor is only available for UScriptStructs */
+	template<typename T = TBaseStruct, typename = typename TEnableIf<!TIsDerivedFrom<T, UObject>::IsDerived>::Type>
 	explicit TStructTypeBitSet(TConstArrayView<FInstancedStruct> InitList)
 	{
 		for (const FInstancedStruct& StructInstance : InitList)
@@ -299,7 +293,7 @@ private:
 		StructTypesBitArray.AddAtIndex(BitToSet);
 	}
 
-	static const UStruct* GetBaseUStruct()
+	FORCEINLINE_DEBUGGABLE static const UStruct* GetBaseUStruct()
 	{
 		static const UStruct* Instance = UE::StructUtils::GetAsUStruct<TBaseStruct>();
 		return Instance;
@@ -307,7 +301,7 @@ private:
 
 public:
 
-	static int32 CreateTypeIndex(const UStruct& InStructType)
+	static int32 CreateTypeIndex(const TUStructType& InStructType)
 	{
 #if WITH_STRUCTUTILS_DEBUG
 		ensureMsgf(InStructType.IsChildOf(GetBaseUStruct())
@@ -327,10 +321,10 @@ public:
 	}
 
 	template<typename T>
-	static const TStructTypeBitSet<TBaseStruct>& GetTypeBitSet()
+	static const TStructTypeBitSet& GetTypeBitSet()
 	{
 		static_assert(TIsDerivedFrom<T, TBaseStruct>::IsDerived, "Given struct type doesn't match the expected base struct type.");
-		static const TStructTypeBitSet<TBaseStruct> TypeBitSet(GetTypeIndex<T>());
+		static const TStructTypeBitSet TypeBitSet(GetTypeIndex<T>());
 		return TypeBitSet;
 	}
 	
@@ -350,7 +344,7 @@ public:
 		StructTypesBitArray.RemoveAtIndex(StructTypeIndex);
 	}
 
-	FORCEINLINE void Remove(const TStructTypeBitSet<TBaseStruct>& Other)
+	FORCEINLINE void Remove(const TStructTypeBitSet& Other)
 	{
 		StructTypesBitArray -= Other.StructTypesBitArray;
 	}
@@ -363,7 +357,7 @@ public:
 		return StructTypesBitArray.Contains(StructTypeIndex);
 	}
 
-	void Add(const UStruct& InStructType)
+	void Add(const TUStructType& InStructType)
 	{
 #if WITH_STRUCTUTILS_DEBUG
 		ensureMsgf(InStructType.IsChildOf(GetBaseUStruct())
@@ -375,7 +369,7 @@ public:
 		StructTypesBitArray.AddAtIndex(StructTypeIndex);
 	}
 
-	void Remove(const UStruct& InStructType)
+	void Remove(const TUStructType& InStructType)
 	{
 #if WITH_STRUCTUTILS_DEBUG
 		ensureMsgf(InStructType.IsChildOf(GetBaseUStruct())
@@ -389,7 +383,7 @@ public:
 
 	void Reset() { StructTypesBitArray.Reset(); }
 
-	bool Contains(const UStruct& InStructType) const
+	bool Contains(const TUStructType& InStructType) const
 	{
 #if WITH_STRUCTUTILS_DEBUG
 		ensureMsgf(InStructType.IsChildOf(GetBaseUStruct())
@@ -418,14 +412,14 @@ public:
 		StructTypesBitArray -= Other.StructTypesBitArray;
 	}
 
-	FORCEINLINE TStructTypeBitSet operator+(const UStruct& NewElement) const
+	FORCEINLINE TStructTypeBitSet operator+(const TUStructType& NewElement) const
 	{
 		TStructTypeBitSet Result = *this;
 		Result.Add(NewElement);
 		return MoveTemp(Result);
 	}
 
-	FORCEINLINE TStructTypeBitSet operator-(const UStruct& NewElement) const
+	FORCEINLINE TStructTypeBitSet operator-(const TUStructType& NewElement) const
 	{
 		TStructTypeBitSet Result = *this;
 		Result.Remove(NewElement);
@@ -454,7 +448,7 @@ public:
 		return *this & Other;
 	}
 
-	FORCEINLINE bool IsEquivalent(const TStructTypeBitSet<TBaseStruct>& Other) const
+	FORCEINLINE bool IsEquivalent(const TStructTypeBitSet& Other) const
 	{
 		return StructTypesBitArray.CompareSetBits(Other.StructTypesBitArray, /*bMissingBitValue=*/false);
 	}
@@ -544,9 +538,9 @@ public:
 		}
 	}
 
-	static TConstArrayView<TWeakObjectPtr<const UStruct>> DebugGetAllStructTypes()
+	static TConstArrayView<TWeakObjectPtr<const TUStructType>> DebugGetAllStructTypes()
 	{
-		return StructTracker.DebugGetAllStructTypes();
+		return StructTracker.DebugGetAllStructTypes<TUStructType>();
 	}
 
 	/**
@@ -564,9 +558,9 @@ protected:
 #endif // WITH_STRUCTUTILS_DEBUG
 
 public:
-	FORCEINLINE friend uint32 GetTypeHash(const TStructTypeBitSet<TBaseStruct>& Instance)
+	FORCEINLINE friend uint32 GetTypeHash(const TStructTypeBitSet<TBaseStruct, TUStructType>& Instance)
 	{
-		const uint32 StoredTypeHash = PointerHash(UE::StructUtils::GetAsUStruct<TBaseStruct>());
+		const uint32 StoredTypeHash = PointerHash(GetBaseUStruct());
 		const uint32 BitArrayHash = GetTypeHash(Instance.StructTypesBitArray);
 		return HashCombine(StoredTypeHash, BitArrayHash);
 	}
@@ -575,4 +569,4 @@ private:
 	FBitArrayExt StructTypesBitArray;
 };
 
-template<typename TBaseStruct> FStructTracker TStructTypeBitSet<TBaseStruct>::StructTracker;
+template<typename TBaseStruct, typename TUStructType> FStructTracker TStructTypeBitSet<TBaseStruct, TUStructType>::StructTracker;
diff --git a/Engine/Plugins/Runtime/MassEntity/Source/MassEntity/Private/MassEntityDebug.cpp b/Engine/Plugins/Runtime/MassEntity/Source/MassEntity/Private/MassEntityDebug.cpp
--- a/Engine/Plugins/Runtime/MassEntity/Source/MassEntity/Private/MassEntityDebug.cpp
+++ b/Engine/Plugins/Runtime/MassEntity/Source/MassEntity/Private/MassEntityDebug.cpp
@@ -135,7 +135,7 @@ FAutoConsoleCommandWithWorldArgsAndOutputDevice LogFragmentSizes(
 	TEXT("Logs all the fragment types being used along with their sizes."),
 	FConsoleCommandWithWorldArgsAndOutputDeviceDelegate::CreateLambda([](const TArray<FString>& Params, UWorld* World, FOutputDevice& Ar)
 		{
-			for (const TWeakObjectPtr<const UStruct>& WeakStruct : FMassFragmentBitSet::DebugGetAllStructTypes())
+			for (const TWeakObjectPtr<const UScriptStruct>& WeakStruct : FMassFragmentBitSet::DebugGetAllStructTypes())
 			{
 				if (const UStruct* StructType = WeakStruct.Get())
 				{
@@ -164,7 +164,7 @@ FAutoConsoleCommandWithOutputDevice LogFragments(
 	TEXT("Logs all the known tags and fragments along with their \"index\" as stored via bitsets."),
 	FConsoleCommandWithOutputDeviceDelegate::CreateStatic([](FOutputDevice& OutputDevice)
 {
-	auto PrintKnownTypes = [&OutputDevice](TConstArrayView<TWeakObjectPtr<const UStruct>> AllStructs) {
+	auto PrintKnownTypes = [&OutputDevice](TConstArrayView<TWeakObjectPtr<const UScriptStruct>> AllStructs) {
 		int i = 0;
 		for (TWeakObjectPtr<const UStruct> Struct : AllStructs)
 		{
