Fix to StructTypeBitSet.StructTracker being instanced per module.

diff --git a/Engine/Plugins/Experimental/StructUtils/Source/StructUtils/Public/StructTypeBitSet.h b/Engine/Plugins/Experimental/StructUtils/Source/StructUtils/Public/StructTypeBitSet.h
--- a/Engine/Plugins/Experimental/StructUtils/Source/StructUtils/Public/StructTypeBitSet.h
+++ b/Engine/Plugins/Experimental/StructUtils/Source/StructUtils/Public/StructTypeBitSet.h
@@ -569,4 +569,19 @@ private:
 	FBitArrayExt StructTypesBitArray;
 };
 
-template<typename TBaseStruct, typename TUStructType> FStructTracker TStructTypeBitSet<TBaseStruct, TUStructType>::StructTracker;
+#define DECLARE_STRUCTTYPEBITSET_EXPORTED(EXPORTED_API, ContainerTypeName, BaseStructType) template<> \
+	FStructTracker TStructTypeBitSet<BaseStructType, UScriptStruct>::StructTracker; \
+	template struct EXPORTED_API TStructTypeBitSet<BaseStructType, UScriptStruct>; \
+	using ContainerTypeName = TStructTypeBitSet<BaseStructType, UScriptStruct>;
+
+#define DECLARE_STRUCTTYPEBITSET(ContainerTypeName, BaseStructType) DECLARE_STRUCTTYPEBITSET_EXPORTED(, ContainerTypeName, BaseStructType)
+
+#define DECLARE_CLASSTYPEBITSET_EXPORTED(EXPORTED_API, ContainerTypeName, BaseStructType) template<> \
+	FStructTracker TStructTypeBitSet<BaseStructType, UClass>::StructTracker; \
+	template struct EXPORTED_API TStructTypeBitSet<BaseStructType, UClass>; \
+	using ContainerTypeName = TStructTypeBitSet<BaseStructType, UClass>;
+
+#define DECLARE_CLASSTYPEBITSET(ContainerTypeName, BaseStructType) DECLARE_STRUCTTYPEBITSET_EXPORTED(, ContainerTypeName, BaseStructType)
+
+#define DEFINE_TYPEBITSET(ContainerTypeName) template<> \
+	FStructTracker ContainerTypeName::StructTracker = FStructTracker();
diff --git a/Engine/Plugins/Experimental/StructUtils/Source/StructUtilsTestSuite/Private/ScriptStructTypeBitSetTest.cpp b/Engine/Plugins/Experimental/StructUtils/Source/StructUtilsTestSuite/Private/ScriptStructTypeBitSetTest.cpp
--- a/Engine/Plugins/Experimental/StructUtils/Source/StructUtilsTestSuite/Private/ScriptStructTypeBitSetTest.cpp
+++ b/Engine/Plugins/Experimental/StructUtils/Source/StructUtilsTestSuite/Private/ScriptStructTypeBitSetTest.cpp
@@ -11,11 +11,11 @@
 
 PRAGMA_DISABLE_OPTIMIZATION
 
+DECLARE_STRUCTTYPEBITSET(FTestStructBaseBitSet, FTestStructSimpleBase);
+DEFINE_TYPEBITSET(FTestStructBaseBitSet);
+
 namespace FStructTypeBitSetTests
 {
-
-using FTestStructBaseBitSet = TStructTypeBitSet<FTestStructSimpleBase>;
-
 struct FTestStructBitSet : public TStructTypeBitSet<FTestStructSimple>
 {
 	FTestStructBitSet() = default;
diff --git a/Engine/Plugins/Runtime/MassEntity/Source/MassEntity/Private/MassEntityTypes.cpp b/Engine/Plugins/Runtime/MassEntity/Source/MassEntity/Private/MassEntityTypes.cpp
--- a/Engine/Plugins/Runtime/MassEntity/Source/MassEntity/Private/MassEntityTypes.cpp
+++ b/Engine/Plugins/Runtime/MassEntity/Source/MassEntity/Private/MassEntityTypes.cpp
@@ -3,6 +3,11 @@
 #include "MassEntityTypes.h"
 #include "StructUtilsTypes.h"
 
+DEFINE_TYPEBITSET(FMassFragmentBitSet);
+DEFINE_TYPEBITSET(FMassTagBitSet);
+DEFINE_TYPEBITSET(FMassChunkFragmentBitSet);
+DEFINE_TYPEBITSET(FMassSharedFragmentBitSet);
+
 uint32 FMassArchetypeSharedFragmentValues::CalculateHash() const
 {
 	checkf(bSorted, TEXT("Expecting the containers to be sorted for the hash caluclation to be correct"));
diff --git a/Engine/Plugins/Runtime/MassEntity/Source/MassEntity/Public/MassEntityTypes.h b/Engine/Plugins/Runtime/MassEntity/Source/MassEntity/Public/MassEntityTypes.h
--- a/Engine/Plugins/Runtime/MassEntity/Source/MassEntity/Public/MassEntityTypes.h
+++ b/Engine/Plugins/Runtime/MassEntity/Source/MassEntity/Public/MassEntityTypes.h
@@ -102,14 +102,10 @@ struct FMassEntityHandle
 	}
 };
 
-template struct MASSENTITY_API TStructTypeBitSet<FMassFragment>;
-using FMassFragmentBitSet = TStructTypeBitSet<FMassFragment>;
-template struct MASSENTITY_API TStructTypeBitSet<FMassTag>;
-using FMassTagBitSet = TStructTypeBitSet<FMassTag>;
-template struct MASSENTITY_API TStructTypeBitSet<FMassChunkFragment>;
-using FMassChunkFragmentBitSet = TStructTypeBitSet<FMassChunkFragment>;
-template struct MASSENTITY_API TStructTypeBitSet<FMassSharedFragment>;
-using FMassSharedFragmentBitSet = TStructTypeBitSet<FMassSharedFragment>;
+DECLARE_STRUCTTYPEBITSET_EXPORTED(MASSENTITY_API, FMassFragmentBitSet, FMassFragment);
+DECLARE_STRUCTTYPEBITSET_EXPORTED(MASSENTITY_API, FMassTagBitSet, FMassTag);
+DECLARE_STRUCTTYPEBITSET_EXPORTED(MASSENTITY_API, FMassChunkFragmentBitSet, FMassChunkFragment);
+DECLARE_STRUCTTYPEBITSET_EXPORTED(MASSENTITY_API, FMassSharedFragmentBitSet, FMassSharedFragment);
 
 /** The type summarily describing a composition of an entity or an archetype. It contains information on both the
  *  fragments as well as tags */
