Fix non-unity no-pch issues
#rnx

diff --git a/Engine/Plugins/Developer/Concert/ConcertApp/MultiUserServer/Source/MultiUserServer/Private/Widgets/SessionTabs/Live/PackageViewer/SConcertSessionPackageViewer.cpp b/Engine/Plugins/Developer/Concert/ConcertApp/MultiUserServer/Source/MultiUserServer/Private/Widgets/SessionTabs/Live/PackageViewer/SConcertSessionPackageViewer.cpp
--- a/Engine/Plugins/Developer/Concert/ConcertApp/MultiUserServer/Source/MultiUserServer/Private/Widgets/SessionTabs/Live/PackageViewer/SConcertSessionPackageViewer.cpp
+++ b/Engine/Plugins/Developer/Concert/ConcertApp/MultiUserServer/Source/MultiUserServer/Private/Widgets/SessionTabs/Live/PackageViewer/SConcertSessionPackageViewer.cpp
@@ -8,6 +8,7 @@
 #include "Session/Activity/PredefinedActivityColumns.h"
 #include "Session/Activity/SConcertSessionActivities.h"
 #include "Widgets/Input/SSearchBox.h"
+#include "Framework/MultiBox/MultiBoxBuilder.h"
 
 #define LOCTEXT_NAMESPACE "UnrealMultiUserUI"
 
diff --git a/Engine/Plugins/Developer/Concert/ConcertUI/ConcertSharedSlate/Source/ConcertSharedSlate/Private/Session/History/SSessionHistory.cpp b/Engine/Plugins/Developer/Concert/ConcertUI/ConcertSharedSlate/Source/ConcertSharedSlate/Private/Session/History/SSessionHistory.cpp
--- a/Engine/Plugins/Developer/Concert/ConcertUI/ConcertSharedSlate/Source/ConcertSharedSlate/Private/Session/History/SSessionHistory.cpp
+++ b/Engine/Plugins/Developer/Concert/ConcertUI/ConcertSharedSlate/Source/ConcertSharedSlate/Private/Session/History/SSessionHistory.cpp
@@ -8,6 +8,7 @@
 #include "Session/Activity/PredefinedActivityColumns.h"
 #include "Session/Activity/SConcertSessionActivities.h"
 
+#include "Framework/MultiBox/MultiBoxBuilder.h"
 #include "Widgets/Input/SSearchBox.h"
 #include "Widgets/Layout/SSeparator.h"
 
diff --git a/Engine/Plugins/Developer/Concert/ConcertUI/ConcertSharedSlate/Source/ConcertSharedSlate/Public/ConcertHeaderRowUtils.h b/Engine/Plugins/Developer/Concert/ConcertUI/ConcertSharedSlate/Source/ConcertSharedSlate/Public/ConcertHeaderRowUtils.h
--- a/Engine/Plugins/Developer/Concert/ConcertUI/ConcertSharedSlate/Source/ConcertSharedSlate/Public/ConcertHeaderRowUtils.h
+++ b/Engine/Plugins/Developer/Concert/ConcertUI/ConcertSharedSlate/Source/ConcertSharedSlate/Public/ConcertHeaderRowUtils.h
@@ -7,6 +7,7 @@
 
 class FMenuBuilder;
 class SHeaderRow;
+class SWidget;
 
 USTRUCT()
 struct CONCERTSHAREDSLATE_API FColumnVisibilitySnapshot
diff --git a/Engine/Plugins/Experimental/Animation/PoseSearch/Source/Runtime/Public/PoseSearch/PoseSearch.h b/Engine/Plugins/Experimental/Animation/PoseSearch/Source/Runtime/Public/PoseSearch/PoseSearch.h
--- a/Engine/Plugins/Experimental/Animation/PoseSearch/Source/Runtime/Public/PoseSearch/PoseSearch.h
+++ b/Engine/Plugins/Experimental/Animation/PoseSearch/Source/Runtime/Public/PoseSearch/PoseSearch.h
@@ -12,6 +12,7 @@
 #include "Animation/AnimMetaData.h"
 #include "Animation/AnimNodeMessages.h"
 #include "Animation/AnimationAsset.h"
+#include "Animation/AnimSequence.h"
 #include "Animation/MotionTrajectoryTypes.h"
 #include "AlphaBlend.h"
 #include "BoneIndices.h"
diff --git a/Engine/Plugins/Experimental/StructUtils/Source/StructUtils/Private/PropertyBag.cpp b/Engine/Plugins/Experimental/StructUtils/Source/StructUtils/Private/PropertyBag.cpp
--- a/Engine/Plugins/Experimental/StructUtils/Source/StructUtils/Private/PropertyBag.cpp
+++ b/Engine/Plugins/Experimental/StructUtils/Source/StructUtils/Private/PropertyBag.cpp
@@ -1,8 +1,9 @@
-﻿// Copyright Epic Games, Inc. All Rights Reserved.
+// Copyright Epic Games, Inc. All Rights Reserved.
 
 #include "PropertyBag.h"
 #include "StructView.h"
 #include "Hash/CityHash.h"
+#include "UObject/TextProperty.h"
 
 namespace UE::StructUtils::Private
 {
diff --git a/Engine/Plugins/Experimental/StructUtils/Source/StructUtilsEditor/Private/Customizations/PropertyBagDetails.cpp b/Engine/Plugins/Experimental/StructUtils/Source/StructUtilsEditor/Private/Customizations/PropertyBagDetails.cpp
--- a/Engine/Plugins/Experimental/StructUtils/Source/StructUtilsEditor/Private/Customizations/PropertyBagDetails.cpp
+++ b/Engine/Plugins/Experimental/StructUtils/Source/StructUtilsEditor/Private/Customizations/PropertyBagDetails.cpp
@@ -1,4 +1,4 @@
-﻿// Copyright Epic Games, Inc. All Rights Reserved.
+// Copyright Epic Games, Inc. All Rights Reserved.
 
 #include "PropertyBagDetails.h"
 #include "InstancedStructDetails.h"
@@ -9,6 +9,7 @@
 #include "IDetailChildrenBuilder.h"
 #include "Widgets/DeclarativeSyntaxSupport.h"
 #include "Widgets/Images/SImage.h"
+#include "Widgets/Input/SButton.h"
 #include "Widgets/Text/SInlineEditableTextBlock.h"
 #include "Widgets/Input/SComboButton.h"
 #include "Framework/MultiBox/MultiBoxBuilder.h"
diff --git a/Engine/Plugins/Experimental/StructUtils/Source/StructUtilsTestSuite/Private/PropertyBagTest.cpp b/Engine/Plugins/Experimental/StructUtils/Source/StructUtilsTestSuite/Private/PropertyBagTest.cpp
--- a/Engine/Plugins/Experimental/StructUtils/Source/StructUtilsTestSuite/Private/PropertyBagTest.cpp
+++ b/Engine/Plugins/Experimental/StructUtils/Source/StructUtilsTestSuite/Private/PropertyBagTest.cpp
@@ -1,8 +1,9 @@
-﻿// Copyright Epic Games, Inc. All Rights Reserved.
+// Copyright Epic Games, Inc. All Rights Reserved.
 
 #include "AITestsCommon.h"
 #include "PropertyBag.h"
 #include "StructUtilsTestTypes.h"
+#include "Engine/World.h"
 
 #define LOCTEXT_NAMESPACE "StructUtilsTests"
 
diff --git a/Engine/Plugins/FX/Niagara/Source/NiagaraEditor/Private/ViewModels/NiagaraScriptGraphViewModel.cpp b/Engine/Plugins/FX/Niagara/Source/NiagaraEditor/Private/ViewModels/NiagaraScriptGraphViewModel.cpp
--- a/Engine/Plugins/FX/Niagara/Source/NiagaraEditor/Private/ViewModels/NiagaraScriptGraphViewModel.cpp
+++ b/Engine/Plugins/FX/Niagara/Source/NiagaraEditor/Private/ViewModels/NiagaraScriptGraphViewModel.cpp
@@ -17,6 +17,7 @@
 #include "NiagaraNodeReroute.h"
 #include "NiagaraScriptVariable.h"
 
+#include "Misc/Base64.h"
 #include "HAL/PlatformApplicationMisc.h"
 #include "Widgets/SNiagaraParameterName.h"
 #include "Widgets/Notifications/SNotificationList.h"
diff --git a/Engine/Plugins/VirtualProduction/LevelSnapshots/Source/LevelSnapshotsEditor/Private/LevelSnapshotsEditorSettings.h b/Engine/Plugins/VirtualProduction/LevelSnapshots/Source/LevelSnapshotsEditor/Private/LevelSnapshotsEditorSettings.h
--- a/Engine/Plugins/VirtualProduction/LevelSnapshots/Source/LevelSnapshotsEditor/Private/LevelSnapshotsEditorSettings.h
+++ b/Engine/Plugins/VirtualProduction/LevelSnapshots/Source/LevelSnapshotsEditor/Private/LevelSnapshotsEditorSettings.h
@@ -3,6 +3,7 @@
 #pragma once
 
 #include "CoreMinimal.h"
+#include "Misc/Attribute.h"
 #include "UObject/Object.h"
 #include "LevelSnapshotsEditorSettings.generated.h"
 
diff --git a/Engine/Source/Runtime/Core/Public/Tasks/TaskPrivate.h b/Engine/Source/Runtime/Core/Public/Tasks/TaskPrivate.h
--- a/Engine/Source/Runtime/Core/Public/Tasks/TaskPrivate.h
+++ b/Engine/Source/Runtime/Core/Public/Tasks/TaskPrivate.h
@@ -6,6 +6,7 @@
 #include "Async/Fundamental/Scheduler.h"
 #include "Async/TaskTrace.h"
 #include "Templates/Invoke.h"
+#include "Templates/RefCounting.h"
 #include "Templates/TypeCompatibleBytes.h"
 #include "Misc/Timeout.h"
 #include "Containers/LockFreeList.h"
diff --git a/Engine/Source/Runtime/CoreUObject/Private/UObject/UObjectBase.cpp b/Engine/Source/Runtime/CoreUObject/Private/UObject/UObjectBase.cpp
--- a/Engine/Source/Runtime/CoreUObject/Private/UObject/UObjectBase.cpp
+++ b/Engine/Source/Runtime/CoreUObject/Private/UObject/UObjectBase.cpp
@@ -350,7 +350,7 @@ bool UObjectBase::IsValidLowLevelFast(bool bRecursive /*= true*/) const
 			return false;
 		}
 	}
-	// Avoid infinite recursion so call IsValidLowLevelFast on the class object with bRecirsive = false.
+	// Avoid infinite recursion so call IsValidLowLevelFast on the class object with bRecursive = false.
 	if (bRecursive && !ClassPrivate->IsValidLowLevelFast(false))
 	{
 		UE_LOG(LogUObjectBase, Error, TEXT("Class object failed IsValidLowLevelFast test."));
diff --git a/Engine/Source/Runtime/CoreUObject/Public/UObject/TopLevelAssetPath.h b/Engine/Source/Runtime/CoreUObject/Public/UObject/TopLevelAssetPath.h
--- a/Engine/Source/Runtime/CoreUObject/Public/UObject/TopLevelAssetPath.h
+++ b/Engine/Source/Runtime/CoreUObject/Public/UObject/TopLevelAssetPath.h
@@ -4,6 +4,7 @@
 
 #include "CoreTypes.h"
 #include "Containers/StringView.h"
+#include "Serialization/StructuredArchive.h"
 #include "UObject/NameTypes.h"
 
 class UObject;
diff --git a/Engine/Source/Runtime/Engine/Private/DeviceProfiles/DeviceProfile.cpp b/Engine/Source/Runtime/Engine/Private/DeviceProfiles/DeviceProfile.cpp
--- a/Engine/Source/Runtime/Engine/Private/DeviceProfiles/DeviceProfile.cpp
+++ b/Engine/Source/Runtime/Engine/Private/DeviceProfiles/DeviceProfile.cpp
@@ -2,6 +2,7 @@
 
 #include "DeviceProfiles/DeviceProfile.h"
 #include "Misc/Paths.h"
+#include "Scalability.h"
 #include "UObject/Package.h"
 #include "UObject/UnrealType.h"
 #include "UObject/UObjectHash.h"
diff --git a/Engine/Source/Runtime/RigVM/Private/RigVMCore/RigVMTemplate.cpp b/Engine/Source/Runtime/RigVM/Private/RigVMCore/RigVMTemplate.cpp
--- a/Engine/Source/Runtime/RigVM/Private/RigVMCore/RigVMTemplate.cpp
+++ b/Engine/Source/Runtime/RigVM/Private/RigVMCore/RigVMTemplate.cpp
@@ -5,6 +5,7 @@
 #include "RigVMCore/RigVMStruct.h"
 #include "RigVMModule.h"
 #include "Algo/Sort.h"
+#include "UObject/UObjectIterator.h"
 
 ////////////////////////////////////////////////////////////////////////////////////////////////
 
