PropertyBag: fix for NoPCH compilation of the tests

diff --git a/Engine/Plugins/Experimental/StructUtils/Source/StructUtilsTestSuite/Private/PropertyBagTest.cpp b/Engine/Plugins/Experimental/StructUtils/Source/StructUtilsTestSuite/Private/PropertyBagTest.cpp
--- a/Engine/Plugins/Experimental/StructUtils/Source/StructUtilsTestSuite/Private/PropertyBagTest.cpp
+++ b/Engine/Plugins/Experimental/StructUtils/Source/StructUtilsTestSuite/Private/PropertyBagTest.cpp
@@ -122,9 +122,9 @@ struct FTest_Object : FAITestBase
 	{
 		static const FName ObjectName(TEXT("Object"));
 
-		UBagTestObject1* Test1 = NewObject<UBagTestObject1>(&GetWorld());
-		UBagTestObject2* Test2 = NewObject<UBagTestObject2>(&GetWorld());
-		UBagTestObject1Derived* Test1Derived = NewObject<UBagTestObject1Derived>(&GetWorld());
+		UBagTestObject1* Test1 = NewObject<UBagTestObject1>();
+		UBagTestObject2* Test2 = NewObject<UBagTestObject2>();
+		UBagTestObject1Derived* Test1Derived = NewObject<UBagTestObject1Derived>();
 
 		FInstancedPropertyBag Bag;
 		Bag.AddProperty(ObjectName, EPropertyBagPropertyType::Object, UBagTestObject1::StaticClass());
