StructTypeBitSet compilation fix for strict conformance mode

diff --git a/Engine/Plugins/Experimental/StructUtils/Source/StructUtils/Public/StructTypeBitSet.h b/Engine/Plugins/Experimental/StructUtils/Source/StructUtils/Public/StructTypeBitSet.h
--- a/Engine/Plugins/Experimental/StructUtils/Source/StructUtils/Public/StructTypeBitSet.h
+++ b/Engine/Plugins/Experimental/StructUtils/Source/StructUtils/Public/StructTypeBitSet.h
@@ -25,7 +25,7 @@
  * The type contains static members so you'll also need to define these. You can easily do it by placing the following 
  * in your cpp file (continuing the FFooBar example):
  *
- *	DEFINE_STRUCTTYPEBITSET(FMyFooBarBitSet);
+ *	DEFINE_TYPEBITSET(FMyFooBarBitSet);
  * 
  */
 struct FStructTracker
@@ -232,8 +232,6 @@ private:
 		}
 	};
 
-	static FStructTracker StructTracker;
-
 public:
 	TStructTypeBitSet() = default;
 
@@ -567,6 +565,8 @@ public:
 
 private:
 	FBitArrayExt StructTypesBitArray;
+
+	static FStructTracker StructTracker;
 };
 
 #define DECLARE_STRUCTTYPEBITSET_EXPORTED(EXPORTED_API, ContainerTypeName, BaseStructType) template<> \
@@ -583,5 +583,9 @@ private:
 
 #define DECLARE_CLASSTYPEBITSET(ContainerTypeName, BaseStructType) DECLARE_STRUCTTYPEBITSET_EXPORTED(, ContainerTypeName, BaseStructType)
 
-#define DEFINE_TYPEBITSET(ContainerTypeName) template<> \
-	FStructTracker ContainerTypeName::StructTracker = FStructTracker();
+#if STRUCTUTILS_STRICT_CONFORMANCE
+	#define DEFINE_TYPEBITSET(ContainerTypeName) template<> \
+		FStructTracker ContainerTypeName::StructTracker = FStructTracker();
+#else
+	#define DEFINE_TYPEBITSET(ContainerTypeName) 
+#endif // STRUCTUTILS_STRICT_CONFORMANCE
diff --git a/Engine/Plugins/Experimental/StructUtils/Source/StructUtils/StructUtils.Build.cs b/Engine/Plugins/Experimental/StructUtils/Source/StructUtils/StructUtils.Build.cs
--- a/Engine/Plugins/Experimental/StructUtils/Source/StructUtils/StructUtils.Build.cs
+++ b/Engine/Plugins/Experimental/StructUtils/Source/StructUtils/StructUtils.Build.cs
@@ -26,7 +26,8 @@ namespace UnrealBuildTool.Rules
 			new string[] {
 			}
 			);
-		}
 
+			PublicDefinitions.Add("STRUCTUTILS_STRICT_CONFORMANCE=" + (Target.WindowsPlatform.bStrictConformanceMode ? "1" : "0"));
+		}
 	}
 }
