Proper solution to static StructTracker compilation/linking issue

diff --git a/Engine/Plugins/Experimental/StructUtils/Source/StructUtils/Public/StructTypeBitSet.h b/Engine/Plugins/Experimental/StructUtils/Source/StructUtils/Public/StructTypeBitSet.h
--- a/Engine/Plugins/Experimental/StructUtils/Source/StructUtils/Public/StructTypeBitSet.h
+++ b/Engine/Plugins/Experimental/StructUtils/Source/StructUtils/Public/StructTypeBitSet.h
@@ -12,8 +12,10 @@
  * The TStructTypeBitSet holds information on "existence" of subtypes of a given UStruct. The information on 
  * available child-structs is gathered lazily - the internal FStructTracker is assigning a given type a new index the
  * very first time the type is encountered. 
- * To create a specific instantiation of the type you need to declare the static StructTracker member variable. 
- * To do it for an arbitrary type FFooBar add the following in your header or cpp file:
+ * To create a specific instantiation of the type you also need to provide a type that will instantiate a static 
+ * FStructTracker instance. The supplied macros hide this details away however.
+ * 
+ * To do declare a bitset type for an arbitrary struct type FFooBar add the following in your header or cpp file:
  * 
  *	DECLARE_STRUCTTYPEBITSET(FMyFooBarBitSet, FFooBar);
  * 
@@ -22,8 +24,8 @@
  * 
  *	DECLARE_STRUCTTYPEBITSET_EXPORTED(MYMODULE_API, FMyFooBarBitSet, FFooBar);
  *
- * The type contains static members so you'll also need to define these. You can easily do it by placing the following 
- * in your cpp file (continuing the FFooBar example):
+ * You also need to instantiate the static FStructTracker added by the DECLARE* macro. You can easily do it by placing 
+ * the following in your cpp file (continuing the FFooBar example):
  *
  *	DEFINE_TYPEBITSET(FMyFooBarBitSet);
  * 
@@ -90,7 +92,7 @@ struct FStructTracker
 	TArray<TWeakObjectPtr<const UStruct>, TInlineAllocator<64>> StructTypesList;
 };
 
-template<typename TBaseStruct, typename TUStructType = UScriptStruct>
+template<typename TBaseStruct, typename TStructTrackerWrapper, typename TUStructType = UScriptStruct>
 struct TStructTypeBitSet
 {
 private:
@@ -279,7 +281,7 @@ private:
 	/** 
 	 * A private constructor for a creating an instance straight from TBitArrays. 
 	 * @Note that this constructor needs to remain private to ensure consistency of stored values with data tracked 
-	 * by the StructTracker
+	 * by the TStructTrackerWrapper::StructTracker
 	 */
 	TStructTypeBitSet(const TBitArray<>& Source)
 		: StructTypesBitArray(Source)
@@ -307,7 +309,7 @@ public:
 			, *InStructType.GetPathName(), *GetBaseUStruct()->GetName());
 #endif // WITH_STRUCTUTILS_DEBUG
 
-		return StructTracker.FindOrAddStructTypeIndex(InStructType);
+		return TStructTrackerWrapper::StructTracker.FindOrAddStructTypeIndex(InStructType);
 	}
 
 	template<typename T>
@@ -363,7 +365,7 @@ public:
 				, *InStructType.GetPathName(), *GetBaseUStruct()->GetName());
 #endif // WITH_STRUCTUTILS_DEBUG
 
-		const int32 StructTypeIndex = StructTracker.FindOrAddStructTypeIndex(InStructType);
+		const int32 StructTypeIndex = TStructTrackerWrapper::StructTracker.FindOrAddStructTypeIndex(InStructType);
 		StructTypesBitArray.AddAtIndex(StructTypeIndex);
 	}
 
@@ -375,7 +377,7 @@ public:
 				, *InStructType.GetPathName(), *GetBaseUStruct()->GetName());
 #endif // WITH_STRUCTUTILS_DEBUG
 
-		const int32 StructTypeIndex = StructTracker.FindOrAddStructTypeIndex(InStructType);
+		const int32 StructTypeIndex = TStructTrackerWrapper::StructTracker.FindOrAddStructTypeIndex(InStructType);
 		StructTypesBitArray.RemoveAtIndex(StructTypeIndex);
 	}
 
@@ -389,7 +391,7 @@ public:
 				, *InStructType.GetPathName(), *GetBaseUStruct()->GetName());
 #endif // WITH_STRUCTUTILS_DEBUG
 
-		const int32 StructTypeIndex = StructTracker.FindOrAddStructTypeIndex(InStructType);
+		const int32 StructTypeIndex = TStructTrackerWrapper::StructTracker.FindOrAddStructTypeIndex(InStructType);
 		return StructTypesBitArray.Contains(StructTypeIndex);
 	}
 
@@ -491,7 +493,7 @@ public:
 		{
 			if (It.GetValue())
 			{
-				OutTypes.Add(Cast<TOutStructType>(StructTracker.GetStructType(It.GetIndex())));
+				OutTypes.Add(Cast<TOutStructType>(TStructTrackerWrapper::StructTracker.GetStructType(It.GetIndex())));
 			}
 			++It;
 		}
@@ -520,7 +522,7 @@ public:
 		{
 			if (StructTypesBitArray[Index])
 			{
-				Ar.Logf(TEXT("%s, "), *StructTracker.DebugGetStructTypeName(Index).ToString());
+				Ar.Logf(TEXT("%s, "), *TStructTrackerWrapper::StructTracker.DebugGetStructTypeName(Index).ToString());
 			}
 		}
 	}
@@ -531,14 +533,14 @@ public:
 		{
 			if (StructTypesBitArray[Index])
 			{
-				OutFNames.Add(StructTracker.DebugGetStructTypeName(Index));
+				OutFNames.Add(TStructTrackerWrapper::StructTracker.DebugGetStructTypeName(Index));
 			}
 		}
 	}
 
 	static TConstArrayView<TWeakObjectPtr<const TUStructType>> DebugGetAllStructTypes()
 	{
-		return StructTracker.DebugGetAllStructTypes<TUStructType>();
+		return TStructTrackerWrapper::StructTracker.template DebugGetAllStructTypes<TUStructType>();
 	}
 
 	/**
@@ -547,8 +549,9 @@ public:
 	 */
 	static void DebugResetStructTypeMappingInfo()
 	{
-		StructTracker.DebugResetStructTypeMappingInfo();
+		TStructTrackerWrapper::StructTracker.DebugResetStructTypeMappingInfo();
 	}
+
 protected:
 	// unittesting purposes only
 	const TBitArray<>& DebugGetStructTypesBitArray() const { return StructTypesBitArray; }
@@ -556,7 +559,7 @@ protected:
 #endif // WITH_STRUCTUTILS_DEBUG
 
 public:
-	FORCEINLINE friend uint32 GetTypeHash(const TStructTypeBitSet<TBaseStruct, TUStructType>& Instance)
+	FORCEINLINE friend uint32 GetTypeHash(const TStructTypeBitSet& Instance)
 	{
 		const uint32 StoredTypeHash = PointerHash(GetBaseUStruct());
 		const uint32 BitArrayHash = GetTypeHash(Instance.StructTypesBitArray);
@@ -565,27 +568,25 @@ public:
 
 private:
 	FBitArrayExt StructTypesBitArray;
-
-	static FStructTracker StructTracker;
 };
 
-#define DECLARE_STRUCTTYPEBITSET_EXPORTED(EXPORTED_API, ContainerTypeName, BaseStructType) template<> \
-	FStructTracker TStructTypeBitSet<BaseStructType, UScriptStruct>::StructTracker; \
-	template struct EXPORTED_API TStructTypeBitSet<BaseStructType, UScriptStruct>; \
-	using ContainerTypeName = TStructTypeBitSet<BaseStructType, UScriptStruct>;
-
-#define DECLARE_STRUCTTYPEBITSET(ContainerTypeName, BaseStructType) DECLARE_STRUCTTYPEBITSET_EXPORTED(, ContainerTypeName, BaseStructType)
-
-#define DECLARE_CLASSTYPEBITSET_EXPORTED(EXPORTED_API, ContainerTypeName, BaseStructType) template<> \
-	FStructTracker TStructTypeBitSet<BaseStructType, UClass>::StructTracker; \
-	template struct EXPORTED_API TStructTypeBitSet<BaseStructType, UClass>; \
-	using ContainerTypeName = TStructTypeBitSet<BaseStructType, UClass>;
-
-#define DECLARE_CLASSTYPEBITSET(ContainerTypeName, BaseStructType) DECLARE_STRUCTTYPEBITSET_EXPORTED(, ContainerTypeName, BaseStructType)
-
-#if STRUCTUTILS_STRICT_CONFORMANCE
-	#define DEFINE_TYPEBITSET(ContainerTypeName) template<> \
-		FStructTracker ContainerTypeName::StructTracker = FStructTracker();
-#else
-	#define DEFINE_TYPEBITSET(ContainerTypeName) 
-#endif // STRUCTUTILS_STRICT_CONFORMANCE
+/** 
+ * We're declaring StructTracker this way rather than being a static member variable of TStructTypeBitSet to avoid linking
+ * issues. We're run into all sorts of issues depending on compiler and whether strict mode was on, and this is the best
+ * way we could come up with to solve it. Thankfully the user doesn't need to even know about this class's existence
+ * as long as they are using the macros below.
+ */
+#define _DECLARE_TYPEBITSET_IMPL(EXPORTED_API, ContainerTypeName, BaseType, BaseUStructType) struct ContainerTypeName##StructTracker \
+	{ \
+		EXPORTED_API static FStructTracker StructTracker; \
+	}; \
+	template struct EXPORTED_API TStructTypeBitSet<BaseType, ContainerTypeName##StructTracker, BaseUStructType>; \
+	using ContainerTypeName = TStructTypeBitSet<BaseType, ContainerTypeName##StructTracker, BaseUStructType>
+
+#define DECLARE_STRUCTTYPEBITSET_EXPORTED(EXPORTED_API, ContainerTypeName, BaseStructType) _DECLARE_TYPEBITSET_IMPL(EXPORTED_API, ContainerTypeName, BaseStructType, UScriptStruct)
+#define DECLARE_STRUCTTYPEBITSET(ContainerTypeName, BaseStructType) _DECLARE_TYPEBITSET_IMPL(, ContainerTypeName, BaseStructType, UScriptStruct)
+#define DECLARE_CLASSTYPEBITSET_EXPORTED(EXPORTED_API, ContainerTypeName, BaseStructType) _DECLARE_TYPEBITSET_IMPL(EXPORTED_API, ContainerTypeName, BaseStructType, UClass)
+#define DECLARE_CLASSTYPEBITSET(ContainerTypeName, BaseStructType) _DECLARE_TYPEBITSET_IMPL(, ContainerTypeName, BaseStructType, UClass)
+
+#define DEFINE_TYPEBITSET(ContainerTypeName) \
+	FStructTracker ContainerTypeName##StructTracker::StructTracker;
diff --git a/Engine/Plugins/Experimental/StructUtils/Source/StructUtilsTestSuite/Private/ScriptStructTypeBitSetTest.cpp b/Engine/Plugins/Experimental/StructUtils/Source/StructUtilsTestSuite/Private/ScriptStructTypeBitSetTest.cpp
--- a/Engine/Plugins/Experimental/StructUtils/Source/StructUtilsTestSuite/Private/ScriptStructTypeBitSetTest.cpp
+++ b/Engine/Plugins/Experimental/StructUtils/Source/StructUtilsTestSuite/Private/ScriptStructTypeBitSetTest.cpp
@@ -16,11 +16,18 @@ DEFINE_TYPEBITSET(FTestStructBaseBitSet);
 
 namespace FStructTypeBitSetTests
 {
-struct FTestStructBitSet : public TStructTypeBitSet<FTestStructSimple>
+struct FTestStructBaseBitSetStructTracker
 {
+	static FStructTracker StructTracker;
+};
+FStructTracker FTestStructBaseBitSetStructTracker::StructTracker;
+
+struct FTestStructBitSet : public TStructTypeBitSet<FTestStructSimple, FTestStructBaseBitSetStructTracker>
+{
+	using Super = TStructTypeBitSet<FTestStructSimple, FTestStructBaseBitSetStructTracker>;
 	FTestStructBitSet() = default;
-	FTestStructBitSet(const TStructTypeBitSet<FTestStructSimple>& Other) 
-		: TStructTypeBitSet<FTestStructSimple>(Other)
+	FTestStructBitSet(const Super& Other)
+		: Super(Other)
 	{
 
 	}
diff --git a/Engine/Plugins/Runtime/MassEntity/Source/MassEntity/Public/MassEntityTypes.h b/Engine/Plugins/Runtime/MassEntity/Source/MassEntity/Public/MassEntityTypes.h
--- a/Engine/Plugins/Runtime/MassEntity/Source/MassEntity/Public/MassEntityTypes.h
+++ b/Engine/Plugins/Runtime/MassEntity/Source/MassEntity/Public/MassEntityTypes.h
@@ -420,11 +420,11 @@ namespace UE::Mass
 			Ordinal = Super::Ordinal + 1
 		};
 
-		template<typename TBaseStruct>
-		static void PopulateBitSet(TStructTypeBitSet<TBaseStruct>& OutBitSet)
+		template<typename TBitSetType>
+		static void PopulateBitSet(TBitSetType& OutBitSet)
 		{
 			Super::PopulateBitSet(OutBitSet);
-			OutBitSet += TStructTypeBitSet<TBaseStruct>::template GetTypeBitSet<FType>();
+			OutBitSet += TBitSetType::template GetTypeBitSet<FType>();
 		}
 	};
 		
@@ -438,10 +438,10 @@ namespace UE::Mass
 			Ordinal = 0
 		};
 
-		template<typename TBaseStruct>
-		static void PopulateBitSet(TStructTypeBitSet<TBaseStruct>& OutBitSet)
+		template<typename TBitSetType>
+		static void PopulateBitSet(TBitSetType& OutBitSet)
 		{
-			OutBitSet += TStructTypeBitSet<TBaseStruct>::template GetTypeBitSet<FType>();
+			OutBitSet += TBitSetType::template GetTypeBitSet<FType>();
 		}
 	};
 
