Property Bag: Fixed layout and other improvements
- Changed GetOrCreateFromDescs() to always ensure on errors
- Added method to migrate to another property bag instance
- Added metatag "FixedLayout" which allows to only edit values, not layout
- Changed the details customization to be allowed to embed in other customizations

diff --git a/Engine/Plugins/Experimental/StructUtils/Source/StructUtils/Private/PropertyBag.cpp b/Engine/Plugins/Experimental/StructUtils/Source/StructUtils/Private/PropertyBag.cpp
--- a/Engine/Plugins/Experimental/StructUtils/Source/StructUtils/Private/PropertyBag.cpp
+++ b/Engine/Plugins/Experimental/StructUtils/Source/StructUtils/Private/PropertyBag.cpp
@@ -744,6 +744,15 @@ void FInstancedPropertyBag::MigrateToNewBagStruct(const UPropertyBag* NewBagStru
 	Value = MoveTemp(NewValue);
 }
 
+void FInstancedPropertyBag::MigrateToNewBagInstance(const FInstancedPropertyBag& NewBagInstance)
+{
+	FInstancedStruct NewValue(NewBagInstance.Value);
+
+	UE::StructUtils::Private::CopyMatchingValuesByID(Value, NewValue);
+	
+	Value = MoveTemp(NewValue);
+}
+
 const UPropertyBag* FInstancedPropertyBag::GetPropertyBagStruct() const
 {
 	return Value.IsValid() ? Cast<const UPropertyBag>(Value.GetScriptStruct()) : nullptr;
@@ -1120,7 +1129,7 @@ const UPropertyBag* UPropertyBag::GetOrCreateFromDescs(const TConstArrayView<FPr
 		{
 			if (Desc.ValueTypeObject == nullptr || Desc.ValueTypeObject->GetClass()->IsChildOf(UScriptStruct::StaticClass()) == false)
 			{
-				ensureMsgf(false, TEXT("Struct property must have valid struct type object."));
+				ensureAlwaysMsgf(false, TEXT("Struct property must have valid struct type object."));
 				bAllValid = false;
 			}
 		}
@@ -1128,7 +1137,7 @@ const UPropertyBag* UPropertyBag::GetOrCreateFromDescs(const TConstArrayView<FPr
 		{
 			if (Desc.ValueTypeObject == nullptr || Desc.ValueTypeObject->GetClass()->IsChildOf(UEnum::StaticClass()) == false)
 			{
-				ensureMsgf(false, TEXT("Enum property must have valid enum type object."));
+				ensureAlwaysMsgf(false, TEXT("Enum property must have valid enum type object."));
 				bAllValid = false;
 			}
 		}
@@ -1136,7 +1145,7 @@ const UPropertyBag* UPropertyBag::GetOrCreateFromDescs(const TConstArrayView<FPr
 		{
 			if (Desc.ValueTypeObject == nullptr)
 			{
-				ensureMsgf(false, TEXT("Object property must have valid object type object."));
+				ensureAlwaysMsgf(false, TEXT("Object property must have valid object type object."));
 				bAllValid = false;
 			}
 		}
@@ -1144,7 +1153,7 @@ const UPropertyBag* UPropertyBag::GetOrCreateFromDescs(const TConstArrayView<FPr
 		{
 			if (Desc.ValueTypeObject == nullptr || Desc.ValueTypeObject->GetClass()->IsChildOf(UClass::StaticClass()) == false)
 			{
-				ensureMsgf(false, TEXT("Class property must have valid class type object."));
+				ensureAlwaysMsgf(false, TEXT("Class property must have valid class type object."));
 				bAllValid = false;
 			}
 		}
diff --git a/Engine/Plugins/Experimental/StructUtils/Source/StructUtils/Public/PropertyBag.h b/Engine/Plugins/Experimental/StructUtils/Source/StructUtils/Public/PropertyBag.h
--- a/Engine/Plugins/Experimental/StructUtils/Source/StructUtils/Public/PropertyBag.h
+++ b/Engine/Plugins/Experimental/StructUtils/Source/StructUtils/Public/PropertyBag.h
@@ -95,6 +95,9 @@ struct STRUCTUTILS_API FPropertyBagPropertyDesc
  * The value is stored as a struct, the type of the value is never serialized, instead the composition of the properties
  * is saved with the instance, and the type is recreated on load. The types with same composition of properties share same type (based on hashing).
  *
+ * UPROPERTY() meta tags:
+ *		- FixedLayout: Property types cannot be altered, but values can be. This is useful if e.g. if the bag layout is set byt code.
+ *
  * NOTE: Adding or removing properties to the instance is quite expensive as it will create new UPropertyBag, reallocate memory, and copy all values over. 
  *
  * Example usage, this allows the bag to be configured in the UI:
@@ -141,6 +144,12 @@ struct STRUCTUTILS_API FInstancedPropertyBag
 	FInstancedPropertyBag& operator=(const FInstancedPropertyBag& InOther) = default;
 	FInstancedPropertyBag& operator=(FInstancedPropertyBag&& InOther) = default;
 
+	/** @return true if the instance contains data. */
+	bool IsValid() const
+	{
+		return Value.IsValid();
+	}
+	
 	/** Resets the instance to empty. */
 	void Reset()
 	{
@@ -188,6 +197,13 @@ struct STRUCTUTILS_API FInstancedPropertyBag
 	 */
 	void MigrateToNewBagStruct(const UPropertyBag* NewBagStruct);
 
+	/**
+	 * Changes the type of this bag to the specified other bag, copies base values from the other bag, and migrates existing values over.
+	 * The properties are matched between the bags based on the property ID.
+	 * @param NewBagInstance Reference to the new type.
+	 */
+	void MigrateToNewBagInstance(const FInstancedPropertyBag& NewBagInstance);
+
 	/** @return pointer to the property bag struct. */ 
 	const UPropertyBag* GetPropertyBagStruct() const;
 	
diff --git a/Engine/Plugins/Experimental/StructUtils/Source/StructUtilsEditor/Private/Customizations/PropertyBagDetails.cpp b/Engine/Plugins/Experimental/StructUtils/Source/StructUtilsEditor/Private/Customizations/PropertyBagDetails.cpp
--- a/Engine/Plugins/Experimental/StructUtils/Source/StructUtilsEditor/Private/Customizations/PropertyBagDetails.cpp
+++ b/Engine/Plugins/Experimental/StructUtils/Source/StructUtilsEditor/Private/Customizations/PropertyBagDetails.cpp
@@ -1,4 +1,4 @@
-// Copyright Epic Games, Inc. All Rights Reserved.
+﻿// Copyright Epic Games, Inc. All Rights Reserved.
 
 #include "PropertyBagDetails.h"
 #include "InstancedStructDetails.h"
@@ -70,7 +70,7 @@ const UPropertyBag* GetCommonBagStruct(TSharedPtr<IPropertyHandle> StructPropert
 }
 
 /** @return property descriptors of the property bag struct common to all edited properties. */
-TArray<FPropertyBagPropertyDesc> GetCommonPropertyDescs(TSharedPtr<IPropertyHandle> StructProperty)
+TArray<FPropertyBagPropertyDesc> GetCommonPropertyDescs(const TSharedPtr<IPropertyHandle>& StructProperty)
 {
 	TArray<FPropertyBagPropertyDesc> PropertyDescs;
 	
@@ -83,7 +83,7 @@ TArray<FPropertyBagPropertyDesc> GetCommonPropertyDescs(TSharedPtr<IPropertyHand
 }
 
 /** Creates new property bag struct and sets all properties to use it, migrating over old values. */
-void SetPropertyDescs(TSharedPtr<IPropertyHandle> StructProperty, const TConstArrayView<FPropertyBagPropertyDesc> PropertyDescs)
+void SetPropertyDescs(const TSharedPtr<IPropertyHandle>& StructProperty, const TConstArrayView<FPropertyBagPropertyDesc> PropertyDescs)
 {
 	if (ensure(IsScriptStruct<FInstancedPropertyBag>(StructProperty)))
 	{
@@ -130,7 +130,7 @@ FName GetValidPropertyName(const FString& Name)
 }
 
 /** @return true of the property name is not used yet by the property bag structure common to all edited properties. */
-bool IsUniqueName(const FName NewName, const FName OldName, TSharedPtr<IPropertyHandle> StructProperty)
+bool IsUniqueName(const FName NewName, const FName OldName, const TSharedPtr<IPropertyHandle>& StructProperty)
 {
 	if (NewName == OldName)
 	{
@@ -340,7 +340,7 @@ void SetPropertyDescFromPin(FPropertyBagPropertyDesc& Desc, const FEdGraphPinTyp
 }
 
 template<typename TFunc>
-void ApplyChangesToPropertyDescs(const FText& SessionName, TSharedPtr<IPropertyHandle> StructProperty, IPropertyUtilities* PropUtils, TFunc&& Function)
+void ApplyChangesToPropertyDescs(const FText& SessionName, const TSharedPtr<IPropertyHandle>& StructProperty, IPropertyUtilities* PropUtils, TFunc&& Function)
 {
 	FScopedTransaction Transaction(SessionName);
 	TArray<FPropertyBagPropertyDesc> PropertyDescs = GetCommonPropertyDescs(StructProperty);
@@ -358,6 +358,29 @@ void ApplyChangesToPropertyDescs(const FText& SessionName, TSharedPtr<IPropertyH
 	}
 }
 
+bool CanHaveMemberVariableOfType(const FEdGraphPinType& PinType)
+{
+	if ((PinType.PinCategory == UEdGraphSchema_K2::PC_Struct))
+	{
+		if (const UObject* TypeObject = PinType.PinSubCategoryObject.Get())
+		{
+			if (TypeObject->IsA<UUserDefinedStruct>())
+			{
+				return false;
+			}
+		}
+	}
+	else if ((PinType.PinCategory == UEdGraphSchema_K2::PC_Exec) 
+		|| (PinType.PinCategory == UEdGraphSchema_K2::PC_Wildcard)
+		|| (PinType.PinCategory == UEdGraphSchema_K2::PC_MCDelegate)
+		|| (PinType.PinCategory == UEdGraphSchema_K2::PC_Delegate)
+		|| (PinType.PinCategory == UEdGraphSchema_K2::PC_Interface))
+	{
+		return false;
+	}
+	
+	return true;
+}
 
 } // UE::StructUtils::Private
 
@@ -369,23 +392,25 @@ void ApplyChangesToPropertyDescs(const FText& SessionName, TSharedPtr<IPropertyH
 //  - ChildPropertyHandle a child property if the ValueProperty  
 //----------------------------------------------------------------//
 
-class FPropertyBagInstanceDataDetails : public FInstancedStructDataDetails
+FPropertyBagInstanceDataDetails::FPropertyBagInstanceDataDetails(TSharedPtr<IPropertyHandle> InValueProperty, TSharedPtr<IPropertyHandle> InStructProperty, IPropertyUtilities* InPropUtils, const bool bInFixedLayout)
+	: FInstancedStructDataDetails(InValueProperty)
+	, StructProperty(InStructProperty)
+	, PropUtils(InPropUtils)
+	, bFixedLayout(bInFixedLayout)
 {
-public:
-	FPropertyBagInstanceDataDetails(TSharedPtr<IPropertyHandle> InValueProperty, TSharedPtr<IPropertyHandle> InStructProperty, IPropertyUtilities* InPropUtils)
-		: FInstancedStructDataDetails(InValueProperty)
-		, StructProperty(InStructProperty)
-		, PropUtils(InPropUtils)
-	{
-	}
+	ensure(UE::StructUtils::Private::IsScriptStruct<FInstancedPropertyBag>(StructProperty));
+	ensure(PropUtils != nullptr);
+}
 
-	virtual void OnChildRowAdded(IDetailPropertyRow& ChildRow) override
-	{
-		TSharedPtr<SWidget> NameWidget;
-		TSharedPtr<SWidget> ValueWidget;
-		FDetailWidgetRow Row;
-		ChildRow.GetDefaultWidgets(NameWidget, ValueWidget, Row);
+void FPropertyBagInstanceDataDetails::OnChildRowAdded(IDetailPropertyRow& ChildRow)
+{
+	TSharedPtr<SWidget> NameWidget;
+	TSharedPtr<SWidget> ValueWidget;
+	FDetailWidgetRow Row;
+	ChildRow.GetDefaultWidgets(NameWidget, ValueWidget, Row);
 
+	if (!bFixedLayout)
+	{
 		// Inline editable name
 		TSharedPtr<SInlineEditableTextBlock> InlineWidget = SNew(SInlineEditableTextBlock)
 			.Text(this, &FPropertyBagInstanceDataDetails::GetPropertyName, ChildRow.GetPropertyHandle())
@@ -394,233 +419,208 @@ public:
 			.OnTextCommitted(this, &FPropertyBagInstanceDataDetails::OnChangePropertyName, ChildRow.GetPropertyHandle())
 			.MultiLine(false);
 
-		ChildRow
-			.CustomWidget(/*bShowChildren*/true)
-			.NameContent()
-			[
-				SNew(SComboButton)
-				.OnGetMenuContent(this, &FPropertyBagInstanceDataDetails::OnPropertyNameContent, ChildRow.GetPropertyHandle(), InlineWidget)
-				.ContentPadding(2)
-				.ButtonStyle(FEditorStyle::Get(), "HoverHintOnly")
-				.ButtonContent()
-				[
-					InlineWidget.ToSharedRef()
-				]
-			]
-			.ValueContent()
+		NameWidget = SNew(SComboButton)
+			.OnGetMenuContent(this, &FPropertyBagInstanceDataDetails::OnPropertyNameContent, ChildRow.GetPropertyHandle(), InlineWidget)
+			.ContentPadding(2)
+			.ButtonStyle(FEditorStyle::Get(), "HoverHintOnly")
+			.ButtonContent()
 			[
-				ValueWidget.ToSharedRef()
+				InlineWidget.ToSharedRef()
 			];
 	}
+		
+	ChildRow
+		.CustomWidget(/*bShowChildren*/true)
+		.NameContent()
+		[
+			NameWidget.ToSharedRef()
+		]
+		.ValueContent()
+		[
+			ValueWidget.ToSharedRef()
+		];
+}
 
 
-	FEdGraphPinType OnGetPinInfo(TSharedPtr<IPropertyHandle> ChildPropertyHandle) const
-	{
-		TArray<FPropertyBagPropertyDesc> PropertyDescs = UE::StructUtils::Private::GetCommonPropertyDescs(StructProperty);
-
-		const FProperty* Property = ChildPropertyHandle->GetProperty();
-		if (FPropertyBagPropertyDesc* Desc = PropertyDescs.FindByPredicate([Property](const FPropertyBagPropertyDesc& Desc){ return Desc.CachedProperty == Property; }))
-		{
-			return UE::StructUtils::Private::GetPropertyDescAsPin(*Desc);
-		}
-		
-		return FEdGraphPinType();
-	}
+FEdGraphPinType FPropertyBagInstanceDataDetails::OnGetPinInfo(TSharedPtr<IPropertyHandle> ChildPropertyHandle) const
+{
+	TArray<FPropertyBagPropertyDesc> PropertyDescs = UE::StructUtils::Private::GetCommonPropertyDescs(StructProperty);
 
-	void PinInfoChanged(const FEdGraphPinType& PinType, TSharedPtr<IPropertyHandle> ChildPropertyHandle) const
+	const FProperty* Property = ChildPropertyHandle->GetProperty();
+	if (FPropertyBagPropertyDesc* Desc = PropertyDescs.FindByPredicate([Property](const FPropertyBagPropertyDesc& Desc){ return Desc.CachedProperty == Property; }))
 	{
-		UE::StructUtils::Private::ApplyChangesToPropertyDescs(
-			LOCTEXT("OnPropertyTypeChanged", "Change Property Type"), StructProperty, PropUtils,
-			[&PinType, &ChildPropertyHandle](TArray<FPropertyBagPropertyDesc>& PropertyDescs)
-			{
-				// Find and change struct type
-				const FProperty* Property = ChildPropertyHandle->GetProperty();
-				if (FPropertyBagPropertyDesc* Desc = PropertyDescs.FindByPredicate([Property](const FPropertyBagPropertyDesc& Desc){ return Desc.CachedProperty == Property; }))
-				{
-					UE::StructUtils::Private::SetPropertyDescFromPin(*Desc, PinType);
-				}
-			});
+		return UE::StructUtils::Private::GetPropertyDescAsPin(*Desc);
 	}
+	
+	return FEdGraphPinType();
+}
 
-	static bool CanHaveMemberVariableOfType(const FEdGraphPinType& PinType)
-	{
-		if ((PinType.PinCategory == UEdGraphSchema_K2::PC_Struct))
+void FPropertyBagInstanceDataDetails::PinInfoChanged(const FEdGraphPinType& PinType, TSharedPtr<IPropertyHandle> ChildPropertyHandle) const
+{
+	UE::StructUtils::Private::ApplyChangesToPropertyDescs(
+		LOCTEXT("OnPropertyTypeChanged", "Change Property Type"), StructProperty, PropUtils,
+		[&PinType, &ChildPropertyHandle](TArray<FPropertyBagPropertyDesc>& PropertyDescs)
 		{
-			if (const UObject* TypeObject = PinType.PinSubCategoryObject.Get())
+			// Find and change struct type
+			const FProperty* Property = ChildPropertyHandle->GetProperty();
+			if (FPropertyBagPropertyDesc* Desc = PropertyDescs.FindByPredicate([Property](const FPropertyBagPropertyDesc& Desc){ return Desc.CachedProperty == Property; }))
 			{
-				if (TypeObject->IsA<UUserDefinedStruct>())
-				{
-					return false;
-				}
+				UE::StructUtils::Private::SetPropertyDescFromPin(*Desc, PinType);
 			}
-		}
-		else if ((PinType.PinCategory == UEdGraphSchema_K2::PC_Exec) 
-			|| (PinType.PinCategory == UEdGraphSchema_K2::PC_Wildcard)
-			|| (PinType.PinCategory == UEdGraphSchema_K2::PC_MCDelegate)
-			|| (PinType.PinCategory == UEdGraphSchema_K2::PC_Delegate)
-			|| (PinType.PinCategory == UEdGraphSchema_K2::PC_Interface))
+		});
+}
+
+void FPropertyBagInstanceDataDetails::GetFilteredVariableTypeTree(TArray<TSharedPtr<UEdGraphSchema_K2::FPinTypeTreeInfo>>& TypeTree, ETypeTreeFilter TypeTreeFilter) const
+{
+	check(GetDefault<UEdGraphSchema_K2>());
+	GetDefault<UEdGraphSchema_K2>()->GetVariableTypeTree(TypeTree, TypeTreeFilter);
+
+	// Filter
+	for (TSharedPtr<UEdGraphSchema_K2::FPinTypeTreeInfo>& PinType : TypeTree)
+	{
+		if (!PinType.IsValid())
 		{
-			return false;
+			return;
 		}
-		return true;
-	}
-	
-	void GetFilteredVariableTypeTree(TArray<TSharedPtr<UEdGraphSchema_K2::FPinTypeTreeInfo>>& TypeTree, ETypeTreeFilter TypeTreeFilter) const
-	{
-		check(GetDefault<UEdGraphSchema_K2>());
-		GetDefault<UEdGraphSchema_K2>()->GetVariableTypeTree(TypeTree, TypeTreeFilter);
 
-		// Filter
-		for (TSharedPtr<UEdGraphSchema_K2::FPinTypeTreeInfo>& PinType : TypeTree)
+		for (int32 ChildIndex = 0; ChildIndex < PinType->Children.Num(); )
 		{
-			if (!PinType.IsValid())
+			TSharedPtr<UEdGraphSchema_K2::FPinTypeTreeInfo> Child = PinType->Children[ChildIndex];
+			if (Child.IsValid())
 			{
-				return;
-			}
-
-			for (int32 ChildIndex = 0; ChildIndex < PinType->Children.Num(); )
-			{
-				TSharedPtr<UEdGraphSchema_K2::FPinTypeTreeInfo> Child = PinType->Children[ChildIndex];
-				if (Child.IsValid())
+				if (!UE::StructUtils::Private::CanHaveMemberVariableOfType(Child->GetPinType(/*bForceLoadSubCategoryObject*/false)))
 				{
-					if (!CanHaveMemberVariableOfType(Child->GetPinType(/*bForceLoadSubCategoryObject*/false)))
-					{
-						PinType->Children.RemoveAt(ChildIndex);
-						continue;
-					}
+					PinType->Children.RemoveAt(ChildIndex);
+					continue;
 				}
-				++ChildIndex;
 			}
-		}			
-	}
+			++ChildIndex;
+		}
+	}			
+}
 
-	
-	FText GetPropertyName(TSharedPtr<IPropertyHandle> ChildPropertyHandle) const
+FText FPropertyBagInstanceDataDetails::GetPropertyName(TSharedPtr<IPropertyHandle> ChildPropertyHandle) const
+{
+	return FText::FromName(ChildPropertyHandle->GetProperty()->GetFName());
+}
+
+bool FPropertyBagInstanceDataDetails::OnValidatePropertyName(const FText& InText, FText& OutErrorMessage, TSharedPtr<IPropertyHandle> ChildPropertyHandle) const
+{
+	const FName NewName = UE::StructUtils::Private::GetValidPropertyName(InText.ToString());
+	bool bResult = UE::StructUtils::Private::IsUniqueName(NewName, ChildPropertyHandle->GetProperty()->GetFName(), StructProperty);
+	if (!bResult)
 	{
-		return FText::FromName(ChildPropertyHandle->GetProperty()->GetFName());
+		OutErrorMessage = LOCTEXT("MustBeUniqueName", "Property must have unique name");
 	}
+	return bResult;
+}
 
-	bool OnValidatePropertyName(const FText& InText, FText& OutErrorMessage, TSharedPtr<IPropertyHandle> ChildPropertyHandle) const
+void FPropertyBagInstanceDataDetails::OnChangePropertyName(const FText& InNewText, ETextCommit::Type InCommitType, TSharedPtr<IPropertyHandle> ChildPropertyHandle) const
+{
+	if (InCommitType == ETextCommit::OnCleared)
 	{
-		const FName NewName = UE::StructUtils::Private::GetValidPropertyName(InText.ToString());
-		bool bResult = UE::StructUtils::Private::IsUniqueName(NewName, ChildPropertyHandle->GetProperty()->GetFName(), StructProperty);
-		if (!bResult)
-		{
-			OutErrorMessage = LOCTEXT("MustBeUniqueName", "Property must have unique name");
-		}
-		return bResult;
+		return;
 	}
 
-	void OnChangePropertyName(const FText& InNewText, ETextCommit::Type InCommitType, TSharedPtr<IPropertyHandle> ChildPropertyHandle) const
+	const FName NewName = UE::StructUtils::Private::GetValidPropertyName(InNewText.ToString());
+	if (!UE::StructUtils::Private::IsUniqueName(NewName, ChildPropertyHandle->GetProperty()->GetFName(), StructProperty))
 	{
-		if (InCommitType == ETextCommit::OnCleared)
-		{
-			return;
-		}
+		return;
+	}
 
-		const FName NewName = UE::StructUtils::Private::GetValidPropertyName(InNewText.ToString());
-		if (!UE::StructUtils::Private::IsUniqueName(NewName, ChildPropertyHandle->GetProperty()->GetFName(), StructProperty))
+	UE::StructUtils::Private::ApplyChangesToPropertyDescs(
+		LOCTEXT("OnPropertyNameChanged", "Change Property Name"), StructProperty, PropUtils,
+		[&NewName, &ChildPropertyHandle](TArray<FPropertyBagPropertyDesc>& PropertyDescs)
 		{
-			return;
-		}
-
-		UE::StructUtils::Private::ApplyChangesToPropertyDescs(
-			LOCTEXT("OnPropertyNameChanged", "Change Property Name"), StructProperty, PropUtils,
-			[&NewName, &ChildPropertyHandle](TArray<FPropertyBagPropertyDesc>& PropertyDescs)
+			const FProperty* Property = ChildPropertyHandle->GetProperty();
+			if (FPropertyBagPropertyDesc* Desc = PropertyDescs.FindByPredicate([Property](const FPropertyBagPropertyDesc& Desc){ return Desc.CachedProperty == Property; }))
 			{
-				const FProperty* Property = ChildPropertyHandle->GetProperty();
-				if (FPropertyBagPropertyDesc* Desc = PropertyDescs.FindByPredicate([Property](const FPropertyBagPropertyDesc& Desc){ return Desc.CachedProperty == Property; }))
-				{
-					Desc->Name = NewName;
-				}
-			});
-	}
+				Desc->Name = NewName;
+			}
+		});
+}
 
-	void OnRemoveProperty(TSharedPtr<IPropertyHandle> ChildPropertyHandle) const
-	{
-		UE::StructUtils::Private::ApplyChangesToPropertyDescs(
-			LOCTEXT("OnPropertyRemoved", "Remove Property"), StructProperty, PropUtils,
-			[&ChildPropertyHandle](TArray<FPropertyBagPropertyDesc>& PropertyDescs)
-			{
-				const FProperty* Property = ChildPropertyHandle->GetProperty();
-				PropertyDescs.RemoveAll([Property](const FPropertyBagPropertyDesc& Desc){ return Desc.CachedProperty == Property; });
-			});
-	}
+void FPropertyBagInstanceDataDetails::OnRemoveProperty(TSharedPtr<IPropertyHandle> ChildPropertyHandle) const
+{
+	UE::StructUtils::Private::ApplyChangesToPropertyDescs(
+		LOCTEXT("OnPropertyRemoved", "Remove Property"), StructProperty, PropUtils,
+		[&ChildPropertyHandle](TArray<FPropertyBagPropertyDesc>& PropertyDescs)
+		{
+			const FProperty* Property = ChildPropertyHandle->GetProperty();
+			PropertyDescs.RemoveAll([Property](const FPropertyBagPropertyDesc& Desc){ return Desc.CachedProperty == Property; });
+		});
+}
 
-	void OnMoveProperty(TSharedPtr<IPropertyHandle> ChildPropertyHandle, int32 Delta) const
-	{
-		UE::StructUtils::Private::ApplyChangesToPropertyDescs(
-			LOCTEXT("OnPropertyMoved", "Move Property"), StructProperty, PropUtils,
-			[&ChildPropertyHandle, &Delta](TArray<FPropertyBagPropertyDesc>& PropertyDescs)
+void FPropertyBagInstanceDataDetails::OnMoveProperty(TSharedPtr<IPropertyHandle> ChildPropertyHandle, int32 Delta) const
+{
+	UE::StructUtils::Private::ApplyChangesToPropertyDescs(
+		LOCTEXT("OnPropertyMoved", "Move Property"), StructProperty, PropUtils,
+		[&ChildPropertyHandle, &Delta](TArray<FPropertyBagPropertyDesc>& PropertyDescs)
+		{
+			// Move
+			if (PropertyDescs.Num() > 1)
 			{
-				// Move
-				if (PropertyDescs.Num() > 1)
+				const FProperty* Property = ChildPropertyHandle->GetProperty();
+				const int32 PropertyIndex = PropertyDescs.IndexOfByPredicate([Property](const FPropertyBagPropertyDesc& Desc){ return Desc.CachedProperty == Property; });
+				if (PropertyIndex != INDEX_NONE)
 				{
-					const FProperty* Property = ChildPropertyHandle->GetProperty();
-					const int32 PropertyIndex = PropertyDescs.IndexOfByPredicate([Property](const FPropertyBagPropertyDesc& Desc){ return Desc.CachedProperty == Property; });
-					if (PropertyIndex != INDEX_NONE)
-					{
-						const int32 NewPropertyIndex = FMath::Clamp(PropertyIndex + Delta, 0, PropertyDescs.Num() - 1);
-						PropertyDescs.Swap(PropertyIndex, NewPropertyIndex);
-					}
+					const int32 NewPropertyIndex = FMath::Clamp(PropertyIndex + Delta, 0, PropertyDescs.Num() - 1);
+					PropertyDescs.Swap(PropertyIndex, NewPropertyIndex);
 				}
-			});
-	}
-
-	TSharedRef<SWidget> OnPropertyNameContent(TSharedPtr<IPropertyHandle> ChildPropertyHandle, TSharedPtr<SInlineEditableTextBlock> InlineWidget) const
-	{
-		constexpr bool bInShouldCloseWindowAfterMenuSelection = true;
-		FMenuBuilder MenuBuilder(bInShouldCloseWindowAfterMenuSelection, nullptr);
-
-		MenuBuilder.AddMenuEntry(
-			FUIAction(),
-			SNew(SPinTypeSelector, FGetPinTypeTree::CreateSP(this, &FPropertyBagInstanceDataDetails::GetFilteredVariableTypeTree))
-				.TargetPinType(this, &FPropertyBagInstanceDataDetails::OnGetPinInfo, ChildPropertyHandle)
-				.OnPinTypeChanged(this, &FPropertyBagInstanceDataDetails::PinInfoChanged, ChildPropertyHandle)
-				.Schema(GetDefault<UEdGraphSchema_K2>())
-				.bAllowArrays(false)
-				.TypeTreeFilter(ETypeTreeFilter::None)
-				.Font( IDetailLayoutBuilder::GetDetailFont() )
-			);
-
-		MenuBuilder.AddSeparator();
-		
-		MenuBuilder.AddMenuEntry(
-			LOCTEXT("Rename", "Rename"),
-			LOCTEXT("Rename_ToolTip", "Rename property"),
-			FSlateIcon(FAppStyle::GetAppStyleSetName(), "Icons.Edit"),
-			FUIAction(FExecuteAction::CreateLambda([InlineWidget]()  { InlineWidget->EnterEditingMode(); }))
-		);
-
-		MenuBuilder.AddMenuEntry(
-			LOCTEXT("Remove", "Remove"),
-			LOCTEXT("Remove_ToolTip", "Remove property"),
-			FSlateIcon(FAppStyle::GetAppStyleSetName(), "Icons.Delete"),
-			FUIAction(FExecuteAction::CreateSP(this, &FPropertyBagInstanceDataDetails::OnRemoveProperty, ChildPropertyHandle))
-		);
-		
-		MenuBuilder.AddSeparator();
-
-		MenuBuilder.AddMenuEntry(
-			LOCTEXT("MoveUp", "Move Up"),
-			LOCTEXT("MoveUp_ToolTip", "Move property up in the list"),
-			FSlateIcon(FAppStyle::GetAppStyleSetName(), "Icons.ArrowUp"),
-			FUIAction(FExecuteAction::CreateSP(this, &FPropertyBagInstanceDataDetails::OnMoveProperty, ChildPropertyHandle, -1))
-		);
+			}
+		});
+}
 
-		MenuBuilder.AddMenuEntry(
-			LOCTEXT("MoveDown", "Move Down"),
-			LOCTEXT("MoveDown_ToolTip", "Move property down in the list"),
-			FSlateIcon(FAppStyle::GetAppStyleSetName(), "Icons.ArrowDown"),
-			FUIAction(FExecuteAction::CreateSP(this, &FPropertyBagInstanceDataDetails::OnMoveProperty, ChildPropertyHandle, +1))
+TSharedRef<SWidget> FPropertyBagInstanceDataDetails::OnPropertyNameContent(TSharedPtr<IPropertyHandle> ChildPropertyHandle, TSharedPtr<SInlineEditableTextBlock> InlineWidget) const
+{
+	constexpr bool bInShouldCloseWindowAfterMenuSelection = true;
+	FMenuBuilder MenuBuilder(bInShouldCloseWindowAfterMenuSelection, nullptr);
+
+	MenuBuilder.AddMenuEntry(
+		FUIAction(),
+		SNew(SPinTypeSelector, FGetPinTypeTree::CreateSP(this, &FPropertyBagInstanceDataDetails::GetFilteredVariableTypeTree))
+			.TargetPinType(this, &FPropertyBagInstanceDataDetails::OnGetPinInfo, ChildPropertyHandle)
+			.OnPinTypeChanged(this, &FPropertyBagInstanceDataDetails::PinInfoChanged, ChildPropertyHandle)
+			.Schema(GetDefault<UEdGraphSchema_K2>())
+			.bAllowArrays(false)
+			.TypeTreeFilter(ETypeTreeFilter::None)
+			.Font( IDetailLayoutBuilder::GetDetailFont() )
 		);
 
-		return MenuBuilder.MakeWidget();
-	}
-
-	TSharedPtr<IPropertyHandle> StructProperty;
-	IPropertyUtilities* PropUtils = nullptr;
-};
+	MenuBuilder.AddSeparator();
+	
+	MenuBuilder.AddMenuEntry(
+		LOCTEXT("Rename", "Rename"),
+		LOCTEXT("Rename_ToolTip", "Rename property"),
+		FSlateIcon(FAppStyle::GetAppStyleSetName(), "Icons.Edit"),
+		FUIAction(FExecuteAction::CreateLambda([InlineWidget]()  { InlineWidget->EnterEditingMode(); }))
+	);
+
+	MenuBuilder.AddMenuEntry(
+		LOCTEXT("Remove", "Remove"),
+		LOCTEXT("Remove_ToolTip", "Remove property"),
+		FSlateIcon(FAppStyle::GetAppStyleSetName(), "Icons.Delete"),
+		FUIAction(FExecuteAction::CreateSP(this, &FPropertyBagInstanceDataDetails::OnRemoveProperty, ChildPropertyHandle))
+	);
+	
+	MenuBuilder.AddSeparator();
+
+	MenuBuilder.AddMenuEntry(
+		LOCTEXT("MoveUp", "Move Up"),
+		LOCTEXT("MoveUp_ToolTip", "Move property up in the list"),
+		FSlateIcon(FAppStyle::GetAppStyleSetName(), "Icons.ArrowUp"),
+		FUIAction(FExecuteAction::CreateSP(this, &FPropertyBagInstanceDataDetails::OnMoveProperty, ChildPropertyHandle, -1))
+	);
+
+	MenuBuilder.AddMenuEntry(
+		LOCTEXT("MoveDown", "Move Down"),
+		LOCTEXT("MoveDown_ToolTip", "Move property down in the list"),
+		FSlateIcon(FAppStyle::GetAppStyleSetName(), "Icons.ArrowDown"),
+		FUIAction(FExecuteAction::CreateSP(this, &FPropertyBagInstanceDataDetails::OnMoveProperty, ChildPropertyHandle, +1))
+	);
+
+	return MenuBuilder.MakeWidget();
+}
 
 
 ////////////////////////////////////
@@ -640,7 +640,18 @@ void FPropertyBagDetails::CustomizeHeader(TSharedRef<class IPropertyHandle> Stru
 	ValueProperty = StructProperty->GetChildHandle(TEXT("Value"));
 	check(ValueProperty);
 
+	static const FName NAME_FixedLayout = "FixedLayout";
+	if (const FProperty* MetaDataProperty = StructProperty->GetMetaDataProperty())
+	{
+		bFixedLayout = MetaDataProperty->HasMetaData(NAME_FixedLayout);
+	}
 
+	TSharedPtr<SWidget> ValueWidget = SNullWidget::NullWidget;
+	if (!bFixedLayout)
+	{
+		ValueWidget = MakeAddPropertyWidget(StructProperty, PropUtils);
+	}
+	
 	HeaderRow
 		.NameContent()
 		[
@@ -649,55 +660,60 @@ void FPropertyBagDetails::CustomizeHeader(TSharedRef<class IPropertyHandle> Stru
 		.ValueContent()
 		.VAlign(VAlign_Center)
 		[
-			SNew(SHorizontalBox)
-			+SHorizontalBox::Slot()
-			.AutoWidth()
-			[
-				SNew(SButton)
-				.VAlign(VAlign_Center)
-				.HAlign(HAlign_Center)
-				.ButtonStyle(FEditorStyle::Get(), "HoverHintOnly")
-				.ToolTipText(LOCTEXT("AddProperty_Tooltip", "Add new property"))
-				.OnClicked(this, &FPropertyBagDetails::OnAddProperty)
-				[
-					SNew(SImage)
-					.Image(FEditorStyle::GetBrush("Icons.PlusCircle"))
-				]
-			]
+			ValueWidget.ToSharedRef()
 		]
 		.ShouldAutoExpand(true);
 }
 
 void FPropertyBagDetails::CustomizeChildren(TSharedRef<class IPropertyHandle> StructPropertyHandle, class IDetailChildrenBuilder& StructBuilder, IPropertyTypeCustomizationUtils& StructCustomizationUtils)
 {
-	// Show the Value (FInstancedStruct) as child rows. 
-	TSharedRef<FPropertyBagInstanceDataDetails> InstanceDetails = MakeShareable(new FPropertyBagInstanceDataDetails(ValueProperty, StructProperty, PropUtils));
+	// Show the Value (FInstancedStruct) as child rows.
+	TSharedRef<FPropertyBagInstanceDataDetails> InstanceDetails = MakeShareable(new FPropertyBagInstanceDataDetails(ValueProperty, StructProperty, PropUtils, bFixedLayout));
 	StructBuilder.AddCustomBuilder(InstanceDetails);
 }
 
-FReply FPropertyBagDetails::OnAddProperty() const
+TSharedPtr<SWidget> FPropertyBagDetails::MakeAddPropertyWidget(TSharedPtr<IPropertyHandle> InStructProperty, class IPropertyUtilities* InPropUtils)
 {
-	constexpr int32 MaxIterations = 100;
-	FName NewName(TEXT("NewProperty"));
-	int32 Number = 1;
-	while (!UE::StructUtils::Private::IsUniqueName(NewName, FName(), StructProperty) && Number < MaxIterations)
-	{
-		Number++;
-		NewName.SetNumber(Number);
-	}
-	if (Number == MaxIterations)
-	{
-		return FReply::Handled();
-	}
+	return SNew(SHorizontalBox)
+		+SHorizontalBox::Slot()
+		.AutoWidth()
+		[
+			SNew(SButton)
+			.VAlign(VAlign_Center)
+			.HAlign(HAlign_Center)
+			.ButtonStyle(FEditorStyle::Get(), "HoverHintOnly")
+			.ToolTipText(LOCTEXT("AddProperty_Tooltip", "Add new property"))
+			.OnClicked_Lambda([InStructProperty, InPropUtils]()
+			{
+				constexpr int32 MaxIterations = 100;
+				FName NewName(TEXT("NewProperty"));
+				int32 Number = 1;
+				while (!UE::StructUtils::Private::IsUniqueName(NewName, FName(), InStructProperty) && Number < MaxIterations)
+				{
+					Number++;
+					NewName.SetNumber(Number);
+				}
+				if (Number == MaxIterations)
+				{
+					return FReply::Handled();
+				}
 
-	UE::StructUtils::Private::ApplyChangesToPropertyDescs(
-		LOCTEXT("OnPropertyAdded", "Add Property"), StructProperty, PropUtils,
-		[&NewName](TArray<FPropertyBagPropertyDesc>& PropertyDescs)
-		{
-			PropertyDescs.Emplace(NewName, EPropertyBagPropertyType::Bool);
-		});
-	
-	return FReply::Handled();
+				UE::StructUtils::Private::ApplyChangesToPropertyDescs(
+					LOCTEXT("OnPropertyAdded", "Add Property"), InStructProperty, InPropUtils,
+					[&NewName](TArray<FPropertyBagPropertyDesc>& PropertyDescs)
+					{
+						PropertyDescs.Emplace(NewName, EPropertyBagPropertyType::Bool);
+					});
+					
+				return FReply::Handled();
+
+			})
+			.Content()
+			[
+				SNew(SImage)
+				.Image(FEditorStyle::GetBrush("Icons.PlusCircle"))
+			]
+		];
 }
 
 #undef LOCTEXT_NAMESPACE
diff --git a/Engine/Plugins/Experimental/StructUtils/Source/StructUtilsEditor/Public/PropertyBagDetails.h b/Engine/Plugins/Experimental/StructUtils/Source/StructUtilsEditor/Public/PropertyBagDetails.h
--- a/Engine/Plugins/Experimental/StructUtils/Source/StructUtilsEditor/Public/PropertyBagDetails.h
+++ b/Engine/Plugins/Experimental/StructUtils/Source/StructUtilsEditor/Public/PropertyBagDetails.h
@@ -3,11 +3,13 @@
 #pragma once
 
 #include "IPropertyTypeCustomization.h"
+#include "InstancedStructDetails.h"
 
 class IPropertyHandle;
+class IDetailPropertyRow;
 
 /**
- * Type customization for FPropertyBag.
+ * Type customization for FInstancedPropertyBag.
  */
 class STRUCTUTILSEDITOR_API FPropertyBagDetails : public IPropertyTypeCustomization
 {
@@ -15,6 +17,9 @@ public:
 	/** Makes a new instance of this detail layout class for a specific detail view requesting it */
 	static TSharedRef<IPropertyTypeCustomization> MakeInstance();
 
+	/** Creates add property widget. */
+	static TSharedPtr<SWidget> MakeAddPropertyWidget(TSharedPtr<IPropertyHandle> StructProperty, class IPropertyUtilities* PropUtils);
+	
 protected:
 	/** IPropertyTypeCustomization interface */
 	virtual void CustomizeHeader(TSharedRef<IPropertyHandle> StructPropertyHandle, class FDetailWidgetRow& HeaderRow, IPropertyTypeCustomizationUtils& StructCustomizationUtils) override;
@@ -28,5 +33,40 @@ private:
 	TSharedPtr<IPropertyHandle> StructProperty;
 	TSharedPtr<IPropertyHandle> ValueProperty;
 
+	bool bFixedLayout = false; 
+	
 	class IPropertyUtilities* PropUtils = nullptr;
 };
+
+/** 
+ * Node builder for FInstancedPropertyBag children.
+ *  - ValueProperty is FInstancedStruct of the FInstancedPropertyBag
+ *  - StructProperty is FInstancedPropertyBag
+ * Can be used in a implementation of a IPropertyTypeCustomization CustomizeChildren() to display editable FInstancedPropertyBag contents.
+ * Use FPropertyBagDetails::MakeAddPropertyWidget() to create the add property widget.
+ * OnChildRowAdded() is called right after each property is added, which allows the property row to be customizable.
+ */
+class STRUCTUTILSEDITOR_API FPropertyBagInstanceDataDetails : public FInstancedStructDataDetails
+{
+public:
+	FPropertyBagInstanceDataDetails(TSharedPtr<IPropertyHandle> InValueProperty, TSharedPtr<IPropertyHandle> InStructProperty, IPropertyUtilities* InPropUtils, const bool bInFixedLayout);
+
+	virtual void OnChildRowAdded(IDetailPropertyRow& ChildRow) override;
+
+protected:
+	FEdGraphPinType OnGetPinInfo(TSharedPtr<IPropertyHandle> ChildPropertyHandle) const;
+	void PinInfoChanged(const FEdGraphPinType& PinType, TSharedPtr<IPropertyHandle> ChildPropertyHandle) const;
+	
+	void GetFilteredVariableTypeTree(TArray<TSharedPtr<UEdGraphSchema_K2::FPinTypeTreeInfo>>& TypeTree, ETypeTreeFilter TypeTreeFilter) const;
+	FText GetPropertyName(TSharedPtr<IPropertyHandle> ChildPropertyHandle) const;
+	bool OnValidatePropertyName(const FText& InText, FText& OutErrorMessage, TSharedPtr<IPropertyHandle> ChildPropertyHandle) const;
+	void OnChangePropertyName(const FText& InNewText, ETextCommit::Type InCommitType, TSharedPtr<IPropertyHandle> ChildPropertyHandle) const;
+	void OnRemoveProperty(TSharedPtr<IPropertyHandle> ChildPropertyHandle) const;
+	void OnMoveProperty(TSharedPtr<IPropertyHandle> ChildPropertyHandle, int32 Delta) const;
+
+	TSharedRef<SWidget> OnPropertyNameContent(TSharedPtr<IPropertyHandle> ChildPropertyHandle, TSharedPtr<SInlineEditableTextBlock> InlineWidget) const;
+
+	TSharedPtr<IPropertyHandle> StructProperty;
+	IPropertyUtilities* PropUtils = nullptr;
+	bool bFixedLayout = false;
+};
