Fixed FInstancedStruct edits not propagating to instances when editing a CDO

Only the ImportText flow of changing a property actually calls PropagatePropertyChange, so we need to use SetPerObjectValues (which calls ImportText internally) rather than directly copy the struct state.

diff --git a/Engine/Plugins/Experimental/StructUtils/Source/StructUtilsEditor/Private/Customizations/InstancedStructDetails.cpp b/Engine/Plugins/Experimental/StructUtils/Source/StructUtilsEditor/Private/Customizations/InstancedStructDetails.cpp
--- a/Engine/Plugins/Experimental/StructUtils/Source/StructUtilsEditor/Private/Customizations/InstancedStructDetails.cpp
+++ b/Engine/Plugins/Experimental/StructUtils/Source/StructUtilsEditor/Private/Customizations/InstancedStructDetails.cpp
@@ -91,6 +91,40 @@ const UScriptStruct* GetCommonScriptStruct(TSharedPtr<IPropertyHandle> StructPro
 	return CommonStructType;
 }
 
+void SetInstancedStructProperty(TSharedPtr<IPropertyHandle> StructProperty, const FInstancedStruct& InstancedStructToSet, const bool bAllowStructMismatch)
+{
+	// Note: We use the ExportText/SetPerObjectValues flow here (rather then CopyScriptStruct) 
+	// so that PropagatePropertyChange is called for the underlying FInstancedStruct property
+	TArray<FString> NewInstancedStructValues;
+	StructProperty->EnumerateRawData([bAllowStructMismatch, &InstancedStructToSet, &NewInstancedStructValues](const void* RawData, const int32 /*DataIndex*/, const int32 /*NumDatas*/)
+	{
+		if (RawData)
+		{
+			const FInstancedStruct* InstancedStruct = static_cast<const FInstancedStruct*>(RawData);
+
+			// Only copy the data if this source is still using the expected struct type or we allow a mismatch
+			const UScriptStruct* StructTypePtr = InstancedStruct->GetScriptStruct();
+			if (bAllowStructMismatch || StructTypePtr == InstancedStructToSet.GetScriptStruct())
+			{
+				FString& NewInstancedStructValue = NewInstancedStructValues.AddDefaulted_GetRef();
+				FInstancedStruct::StaticStruct()->ExportText(NewInstancedStructValue, &InstancedStructToSet, &InstancedStructToSet, nullptr, PPF_None, nullptr);
+			}
+			else
+			{
+				FString& NewInstancedStructValue = NewInstancedStructValues.AddDefaulted_GetRef();
+				FInstancedStruct::StaticStruct()->ExportText(NewInstancedStructValue, InstancedStruct, InstancedStruct, nullptr, PPF_None, nullptr);
+			}
+		}
+		else
+		{
+			NewInstancedStructValues.AddDefaulted();
+		}
+		return true;
+	});
+
+	StructProperty->SetPerObjectValues(NewInstancedStructValues);
+}
+
 }
 
 ////////////////////////////////////
@@ -117,14 +151,21 @@ void FInstancedStructDataDetails::OnStructValuePreChange()
 
 void FInstancedStructDataDetails::OnStructValuePostChange()
 {
-	// Copy the modified struct data back to the source instances
-	SyncEditableInstanceToSource();
-
 	// Forward the change event to the real struct handle
 	if (StructProperty && StructProperty->IsValidHandle())
 	{
 		TGuardValue<bool> HandlingStructValuePostChangeGuard(bIsHandlingStructValuePostChange, true);
 
+		// Copy the modified struct data back to the source instances
+		{
+			FInstancedStruct TmpInstancedStruct;
+			if (StructInstanceData)
+			{
+				TmpInstancedStruct.InitializeAs(Cast<UScriptStruct>(StructInstanceData->GetStruct()), StructInstanceData->GetStructMemory());
+			}
+			UE::StructUtils::Private::SetInstancedStructProperty(StructProperty, TmpInstancedStruct, /*bAllowStructMismatch*/false);
+		}
+
 		StructProperty->NotifyPostChange(EPropertyChangeType::ValueSet);
 		StructProperty->NotifyFinishedChangingProperties();
 	}
@@ -176,41 +217,6 @@ void FInstancedStructDataDetails::SyncEditableInstanceFromSource(bool* OutStruct
 	LastSyncEditableInstanceFromSourceSeconds = FPlatformTime::Seconds();
 }
 
-void FInstancedStructDataDetails::SyncEditableInstanceToSource(bool* OutStructMismatch)
-{
-	if (OutStructMismatch)
-	{
-		*OutStructMismatch = false;
-	}
-
-	if (StructProperty && StructProperty->IsValidHandle())
-	{
-		const UScriptStruct* ExpectedStructType = StructInstanceData ? Cast<UScriptStruct>(StructInstanceData->GetStruct()) : nullptr;
-		StructProperty->EnumerateRawData([this, ExpectedStructType, OutStructMismatch](const void* RawData, const int32 /*DataIndex*/, const int32 /*NumDatas*/)
-		{
-			if (RawData)
-			{
-				const FInstancedStruct* InstancedStruct = static_cast<const FInstancedStruct*>(RawData);
-
-				// Only copy the data if this source is still using the expected struct type
-				const UScriptStruct* StructTypePtr = InstancedStruct->GetScriptStruct();
-				if (StructTypePtr == ExpectedStructType)
-				{
-					if (StructTypePtr)
-					{
-						StructTypePtr->CopyScriptStruct(InstancedStruct->GetMutableMemory(), StructInstanceData->GetStructMemory());
-					}
-				}
-				else if (OutStructMismatch)
-				{
-					*OutStructMismatch = true;
-				}
-			}
-			return false;
-		});
-	}
-}
-
 void FInstancedStructDataDetails::SetOnRebuildChildren(FSimpleDelegate InOnRegenerateChildren)
 {
 	OnRegenerateChildren = InOnRegenerateChildren;
@@ -434,14 +440,12 @@ void FInstancedStructDetails::OnStructPicked(const UScriptStruct* InStruct)
 
 		StructProperty->NotifyPreChange();
 
-		StructProperty->EnumerateRawData([InStruct](void* RawData, const int32 /*DataIndex*/, const int32 /*NumDatas*/)
+		// Copy the modified struct data back to the source instances
 		{
-			if (FInstancedStruct* InstancedStruct = static_cast<FInstancedStruct*>(RawData))
-			{
-				InstancedStruct->InitializeAs(InStruct);
-			}
-			return true;
-		});
+			FInstancedStruct TmpInstancedStruct;
+			TmpInstancedStruct.InitializeAs(InStruct);
+			UE::StructUtils::Private::SetInstancedStructProperty(StructProperty, TmpInstancedStruct, /*bAllowStructMismatch*/true);
+		}
 
 		StructProperty->NotifyPostChange(EPropertyChangeType::ValueSet);
 		StructProperty->NotifyFinishedChangingProperties();
diff --git a/Engine/Plugins/Experimental/StructUtils/Source/StructUtilsEditor/Public/InstancedStructDetails.h b/Engine/Plugins/Experimental/StructUtils/Source/StructUtilsEditor/Public/InstancedStructDetails.h
--- a/Engine/Plugins/Experimental/StructUtils/Source/StructUtilsEditor/Public/InstancedStructDetails.h
+++ b/Engine/Plugins/Experimental/StructUtils/Source/StructUtilsEditor/Public/InstancedStructDetails.h
@@ -75,9 +75,8 @@ private:
 	void OnStructValuePostChange();
 	void OnStructHandlePostChange();
 
-	/** Sync the current state of the editable struct instance to/from the source instance(s) */
+	/** Sync the current state of the editable struct instance from the source instance(s) */
 	void SyncEditableInstanceFromSource(bool* OutStructMismatch = nullptr);
-	void SyncEditableInstanceToSource(bool* OutStructMismatch = nullptr);
 
 	/** Handle to the struct property being edited */
 	TSharedPtr<IPropertyHandle> StructProperty;
