Removed remnants of STRUCTUTILS_STRICT_CONFORMANCE StructTypeBitSet fix.

diff --git a/Engine/Plugins/Experimental/StructUtils/Source/StructUtils/StructUtils.Build.cs b/Engine/Plugins/Experimental/StructUtils/Source/StructUtils/StructUtils.Build.cs
--- a/Engine/Plugins/Experimental/StructUtils/Source/StructUtils/StructUtils.Build.cs
+++ b/Engine/Plugins/Experimental/StructUtils/Source/StructUtils/StructUtils.Build.cs
@@ -21,13 +21,6 @@ namespace UnrealBuildTool.Rules
 				"Engine",
 			}
 			);
-
-			PrivateDependencyModuleNames.AddRange(
-			new string[] {
-			}
-			);
-
-			PublicDefinitions.Add("STRUCTUTILS_STRICT_CONFORMANCE=" + (Target.WindowsPlatform.bStrictConformanceMode ? "1" : "0"));
 		}
 	}
 }
