Fixing various NonUnity compile errors.

diff --git a/Engine/Plugins/Experimental/StructUtils/Source/StructUtilsEditor/Public/PropertyBagDetails.h b/Engine/Plugins/Experimental/StructUtils/Source/StructUtilsEditor/Public/PropertyBagDetails.h
--- a/Engine/Plugins/Experimental/StructUtils/Source/StructUtilsEditor/Public/PropertyBagDetails.h
+++ b/Engine/Plugins/Experimental/StructUtils/Source/StructUtilsEditor/Public/PropertyBagDetails.h
@@ -1,12 +1,14 @@
-﻿// Copyright Epic Games, Inc. All Rights Reserved.
+// Copyright Epic Games, Inc. All Rights Reserved.
 
 #pragma once
 
 #include "IPropertyTypeCustomization.h"
 #include "InstancedStructDetails.h"
+#include "EdGraphSchema_K2.h"
 
 class IPropertyHandle;
 class IDetailPropertyRow;
+class SInlineEditableTextBlock;
 
 /**
  * Type customization for FInstancedPropertyBag.
diff --git a/Engine/Plugins/Media/ImgMedia/Source/ImgMedia/Private/Assets/ImgMediaMipMapInfoManager.h b/Engine/Plugins/Media/ImgMedia/Source/ImgMedia/Private/Assets/ImgMediaMipMapInfoManager.h
--- a/Engine/Plugins/Media/ImgMedia/Source/ImgMedia/Private/Assets/ImgMediaMipMapInfoManager.h
+++ b/Engine/Plugins/Media/ImgMedia/Source/ImgMedia/Private/Assets/ImgMediaMipMapInfoManager.h
@@ -17,7 +17,7 @@ class IMediaPlayer;
 class UCameraComponent;
 class UImgMediaPlaybackComponent;
 class UMediaTexture;
-
+class FViewport;
 
 /**
  * Helps manage FImgMediaMipMapInfo objects.
diff --git a/Engine/Plugins/Media/MediaIOFramework/Source/MediaIOCore/Private/Shared/MediaIOCoreTextureSampleBase.cpp b/Engine/Plugins/Media/MediaIOFramework/Source/MediaIOCore/Private/Shared/MediaIOCoreTextureSampleBase.cpp
--- a/Engine/Plugins/Media/MediaIOFramework/Source/MediaIOCore/Private/Shared/MediaIOCoreTextureSampleBase.cpp
+++ b/Engine/Plugins/Media/MediaIOFramework/Source/MediaIOCore/Private/Shared/MediaIOCoreTextureSampleBase.cpp
@@ -140,3 +140,13 @@ void FMediaIOCoreTextureSampleBase::SetDestructionCallback(TFunction<void(TRefCo
 {
 	DestructionCallback = InDestructionCallback;
 }
+
+void FMediaIOCoreTextureSampleBase::ShutdownPoolable()
+{
+	if (DestructionCallback)
+	{
+		DestructionCallback(Texture);
+	}
+
+	FreeSample();
+}
diff --git a/Engine/Plugins/Media/MediaIOFramework/Source/MediaIOCore/Public/MediaIOCoreTextureSampleBase.h b/Engine/Plugins/Media/MediaIOFramework/Source/MediaIOCore/Public/MediaIOCoreTextureSampleBase.h
--- a/Engine/Plugins/Media/MediaIOFramework/Source/MediaIOCore/Public/MediaIOCoreTextureSampleBase.h
+++ b/Engine/Plugins/Media/MediaIOFramework/Source/MediaIOCore/Public/MediaIOCoreTextureSampleBase.h
@@ -225,15 +225,7 @@ private:
 public:
 	//~ IMediaPoolable interface
 
-	virtual void ShutdownPoolable() override
-	{
-		if (DestructionCallback)
-		{
-			DestructionCallback(Texture);
-		}
-
-		FreeSample();
-	}
+	virtual void ShutdownPoolable() override;
 
 protected:
 	virtual void FreeSample()
diff --git a/Engine/Source/Runtime/Renderer/Private/RectLightTextureManager.cpp b/Engine/Source/Runtime/Renderer/Private/RectLightTextureManager.cpp
--- a/Engine/Source/Runtime/Renderer/Private/RectLightTextureManager.cpp
+++ b/Engine/Source/Runtime/Renderer/Private/RectLightTextureManager.cpp
@@ -17,6 +17,7 @@
 #include "SystemTextures.h"
 #include "TextureLayout.h"
 #include "CommonRenderResources.h"
+#include "ScreenPass.h"
 
 DEFINE_LOG_CATEGORY_STATIC(LogRectLightTextureManager, Log, All);
 
