PropertyBag: Moved internal methods into lambdas to avoid exposing types in header

diff --git a/Engine/Plugins/Experimental/StructUtils/Source/StructUtilsEditor/Private/Customizations/PropertyBagDetails.cpp b/Engine/Plugins/Experimental/StructUtils/Source/StructUtilsEditor/Private/Customizations/PropertyBagDetails.cpp
--- a/Engine/Plugins/Experimental/StructUtils/Source/StructUtilsEditor/Private/Customizations/PropertyBagDetails.cpp
+++ b/Engine/Plugins/Experimental/StructUtils/Source/StructUtilsEditor/Private/Customizations/PropertyBagDetails.cpp
@@ -411,13 +411,50 @@ void FPropertyBagInstanceDataDetails::OnChildRowAdded(IDetailPropertyRow& ChildR
 
 	if (!bFixedLayout)
 	{
+		TSharedPtr<IPropertyHandle> ChildPropertyHandle = ChildRow.GetPropertyHandle();
+
 		// Inline editable name
 		TSharedPtr<SInlineEditableTextBlock> InlineWidget = SNew(SInlineEditableTextBlock)
-			.Text(this, &FPropertyBagInstanceDataDetails::GetPropertyName, ChildRow.GetPropertyHandle())
 			.Font(IDetailLayoutBuilder::GetDetailFont())
-			.OnVerifyTextChanged(this, &FPropertyBagInstanceDataDetails::OnValidatePropertyName, ChildRow.GetPropertyHandle())
-			.OnTextCommitted(this, &FPropertyBagInstanceDataDetails::OnChangePropertyName, ChildRow.GetPropertyHandle())
-			.MultiLine(false);
+			.MultiLine(false)
+			.Text_Lambda([ChildPropertyHandle]()
+			{
+				return FText::FromName(ChildPropertyHandle->GetProperty()->GetFName());
+			})
+			.OnVerifyTextChanged_Lambda([this, ChildPropertyHandle](const FText& InText, FText& OutErrorMessage)
+			{
+				const FName NewName = UE::StructUtils::Private::GetValidPropertyName(InText.ToString());
+				bool bResult = UE::StructUtils::Private::IsUniqueName(NewName, ChildPropertyHandle->GetProperty()->GetFName(), StructProperty);
+				if (!bResult)
+				{
+					OutErrorMessage = LOCTEXT("MustBeUniqueName", "Property must have unique name");
+				}
+				return bResult;
+			})
+			.OnTextCommitted_Lambda([this, ChildPropertyHandle](const FText& InNewText, ETextCommit::Type InCommitType)
+			{
+				if (InCommitType == ETextCommit::OnCleared)
+				{
+					return;
+				}
+
+				const FName NewName = UE::StructUtils::Private::GetValidPropertyName(InNewText.ToString());
+				if (!UE::StructUtils::Private::IsUniqueName(NewName, ChildPropertyHandle->GetProperty()->GetFName(), StructProperty))
+				{
+					return;
+				}
+
+				UE::StructUtils::Private::ApplyChangesToPropertyDescs(
+					LOCTEXT("OnPropertyNameChanged", "Change Property Name"), StructProperty, PropUtils,
+					[&NewName, &ChildPropertyHandle](TArray<FPropertyBagPropertyDesc>& PropertyDescs)
+					{
+						const FProperty* Property = ChildPropertyHandle->GetProperty();
+						if (FPropertyBagPropertyDesc* Desc = PropertyDescs.FindByPredicate([Property](const FPropertyBagPropertyDesc& Desc){ return Desc.CachedProperty == Property; }))
+						{
+							Desc->Name = NewName;
+						}
+					});
+			});
 
 		NameWidget = SNew(SComboButton)
 			.OnGetMenuContent(this, &FPropertyBagInstanceDataDetails::OnPropertyNameContent, ChildRow.GetPropertyHandle(), InlineWidget)
@@ -441,119 +478,82 @@ void FPropertyBagInstanceDataDetails::OnChildRowAdded(IDetailPropertyRow& ChildR
 		];
 }
 
-
-FEdGraphPinType FPropertyBagInstanceDataDetails::OnGetPinInfo(TSharedPtr<IPropertyHandle> ChildPropertyHandle) const
+TSharedRef<SWidget> FPropertyBagInstanceDataDetails::OnPropertyNameContent(TSharedPtr<IPropertyHandle> ChildPropertyHandle, TSharedPtr<SInlineEditableTextBlock> InlineWidget) const
 {
-	TArray<FPropertyBagPropertyDesc> PropertyDescs = UE::StructUtils::Private::GetCommonPropertyDescs(StructProperty);
+	constexpr bool bInShouldCloseWindowAfterMenuSelection = true;
+	FMenuBuilder MenuBuilder(bInShouldCloseWindowAfterMenuSelection, nullptr);
 
-	const FProperty* Property = ChildPropertyHandle->GetProperty();
-	if (FPropertyBagPropertyDesc* Desc = PropertyDescs.FindByPredicate([Property](const FPropertyBagPropertyDesc& Desc){ return Desc.CachedProperty == Property; }))
+	auto GetFilteredVariableTypeTree = [](TArray<TSharedPtr<UEdGraphSchema_K2::FPinTypeTreeInfo>>& TypeTree, ETypeTreeFilter TypeTreeFilter)
 	{
-		return UE::StructUtils::Private::GetPropertyDescAsPin(*Desc);
-	}
-	
-	return FEdGraphPinType();
-}
+		check(GetDefault<UEdGraphSchema_K2>());
+		GetDefault<UEdGraphSchema_K2>()->GetVariableTypeTree(TypeTree, TypeTreeFilter);
 
-void FPropertyBagInstanceDataDetails::PinInfoChanged(const FEdGraphPinType& PinType, TSharedPtr<IPropertyHandle> ChildPropertyHandle) const
-{
-	UE::StructUtils::Private::ApplyChangesToPropertyDescs(
-		LOCTEXT("OnPropertyTypeChanged", "Change Property Type"), StructProperty, PropUtils,
-		[&PinType, &ChildPropertyHandle](TArray<FPropertyBagPropertyDesc>& PropertyDescs)
+		// Filter
+		for (TSharedPtr<UEdGraphSchema_K2::FPinTypeTreeInfo>& PinType : TypeTree)
 		{
-			// Find and change struct type
-			const FProperty* Property = ChildPropertyHandle->GetProperty();
-			if (FPropertyBagPropertyDesc* Desc = PropertyDescs.FindByPredicate([Property](const FPropertyBagPropertyDesc& Desc){ return Desc.CachedProperty == Property; }))
+			if (!PinType.IsValid())
 			{
-				UE::StructUtils::Private::SetPropertyDescFromPin(*Desc, PinType);
+				return;
 			}
-		});
-}
-
-void FPropertyBagInstanceDataDetails::GetFilteredVariableTypeTree(TArray<TSharedPtr<UEdGraphSchema_K2::FPinTypeTreeInfo>>& TypeTree, ETypeTreeFilter TypeTreeFilter) const
-{
-	check(GetDefault<UEdGraphSchema_K2>());
-	GetDefault<UEdGraphSchema_K2>()->GetVariableTypeTree(TypeTree, TypeTreeFilter);
 
-	// Filter
-	for (TSharedPtr<UEdGraphSchema_K2::FPinTypeTreeInfo>& PinType : TypeTree)
-	{
-		if (!PinType.IsValid())
-		{
-			return;
-		}
-
-		for (int32 ChildIndex = 0; ChildIndex < PinType->Children.Num(); )
-		{
-			TSharedPtr<UEdGraphSchema_K2::FPinTypeTreeInfo> Child = PinType->Children[ChildIndex];
-			if (Child.IsValid())
+			for (int32 ChildIndex = 0; ChildIndex < PinType->Children.Num(); )
 			{
-				if (!UE::StructUtils::Private::CanHaveMemberVariableOfType(Child->GetPinType(/*bForceLoadSubCategoryObject*/false)))
+				TSharedPtr<UEdGraphSchema_K2::FPinTypeTreeInfo> Child = PinType->Children[ChildIndex];
+				if (Child.IsValid())
 				{
-					PinType->Children.RemoveAt(ChildIndex);
-					continue;
+					if (!UE::StructUtils::Private::CanHaveMemberVariableOfType(Child->GetPinType(/*bForceLoadSubCategoryObject*/false)))
+					{
+						PinType->Children.RemoveAt(ChildIndex);
+						continue;
+					}
 				}
+				++ChildIndex;
 			}
-			++ChildIndex;
-		}
-	}			
-}
+		}			
+	};
 
-FText FPropertyBagInstanceDataDetails::GetPropertyName(TSharedPtr<IPropertyHandle> ChildPropertyHandle) const
-{
-	return FText::FromName(ChildPropertyHandle->GetProperty()->GetFName());
-}
-
-bool FPropertyBagInstanceDataDetails::OnValidatePropertyName(const FText& InText, FText& OutErrorMessage, TSharedPtr<IPropertyHandle> ChildPropertyHandle) const
-{
-	const FName NewName = UE::StructUtils::Private::GetValidPropertyName(InText.ToString());
-	bool bResult = UE::StructUtils::Private::IsUniqueName(NewName, ChildPropertyHandle->GetProperty()->GetFName(), StructProperty);
-	if (!bResult)
+	auto GetPinInfo = [this, ChildPropertyHandle]()
 	{
-		OutErrorMessage = LOCTEXT("MustBeUniqueName", "Property must have unique name");
-	}
-	return bResult;
-}
+		TArray<FPropertyBagPropertyDesc> PropertyDescs = UE::StructUtils::Private::GetCommonPropertyDescs(StructProperty);
 
-void FPropertyBagInstanceDataDetails::OnChangePropertyName(const FText& InNewText, ETextCommit::Type InCommitType, TSharedPtr<IPropertyHandle> ChildPropertyHandle) const
-{
-	if (InCommitType == ETextCommit::OnCleared)
-	{
-		return;
-	}
+		const FProperty* Property = ChildPropertyHandle->GetProperty();
+		if (FPropertyBagPropertyDesc* Desc = PropertyDescs.FindByPredicate([Property](const FPropertyBagPropertyDesc& Desc){ return Desc.CachedProperty == Property; }))
+		{
+			return UE::StructUtils::Private::GetPropertyDescAsPin(*Desc);
+		}
+	
+		return FEdGraphPinType();
+	};
 
-	const FName NewName = UE::StructUtils::Private::GetValidPropertyName(InNewText.ToString());
-	if (!UE::StructUtils::Private::IsUniqueName(NewName, ChildPropertyHandle->GetProperty()->GetFName(), StructProperty))
+	auto PinInfoChanged = [this, ChildPropertyHandle](const FEdGraphPinType& PinType)
 	{
-		return;
-	}
-
-	UE::StructUtils::Private::ApplyChangesToPropertyDescs(
-		LOCTEXT("OnPropertyNameChanged", "Change Property Name"), StructProperty, PropUtils,
-		[&NewName, &ChildPropertyHandle](TArray<FPropertyBagPropertyDesc>& PropertyDescs)
-		{
-			const FProperty* Property = ChildPropertyHandle->GetProperty();
-			if (FPropertyBagPropertyDesc* Desc = PropertyDescs.FindByPredicate([Property](const FPropertyBagPropertyDesc& Desc){ return Desc.CachedProperty == Property; }))
+		UE::StructUtils::Private::ApplyChangesToPropertyDescs(
+			LOCTEXT("OnPropertyTypeChanged", "Change Property Type"), StructProperty, PropUtils,
+			[&PinType, &ChildPropertyHandle](TArray<FPropertyBagPropertyDesc>& PropertyDescs)
 			{
-				Desc->Name = NewName;
-			}
-		});
-}
+				// Find and change struct type
+				const FProperty* Property = ChildPropertyHandle->GetProperty();
+				if (FPropertyBagPropertyDesc* Desc = PropertyDescs.FindByPredicate([Property](const FPropertyBagPropertyDesc& Desc){ return Desc.CachedProperty == Property; }))
+				{
+					UE::StructUtils::Private::SetPropertyDescFromPin(*Desc, PinType);
+				}
+			});
+	};
 
-void FPropertyBagInstanceDataDetails::OnRemoveProperty(TSharedPtr<IPropertyHandle> ChildPropertyHandle) const
-{
-	UE::StructUtils::Private::ApplyChangesToPropertyDescs(
-		LOCTEXT("OnPropertyRemoved", "Remove Property"), StructProperty, PropUtils,
-		[&ChildPropertyHandle](TArray<FPropertyBagPropertyDesc>& PropertyDescs)
-		{
-			const FProperty* Property = ChildPropertyHandle->GetProperty();
-			PropertyDescs.RemoveAll([Property](const FPropertyBagPropertyDesc& Desc){ return Desc.CachedProperty == Property; });
-		});
-}
+	auto RemoveProperty = [this, ChildPropertyHandle]()
+	{
+		UE::StructUtils::Private::ApplyChangesToPropertyDescs(
+			LOCTEXT("OnPropertyRemoved", "Remove Property"), StructProperty, PropUtils,
+			[&ChildPropertyHandle](TArray<FPropertyBagPropertyDesc>& PropertyDescs)
+			{
+				const FProperty* Property = ChildPropertyHandle->GetProperty();
+				PropertyDescs.RemoveAll([Property](const FPropertyBagPropertyDesc& Desc){ return Desc.CachedProperty == Property; });
+			});
+	};
 
-void FPropertyBagInstanceDataDetails::OnMoveProperty(TSharedPtr<IPropertyHandle> ChildPropertyHandle, int32 Delta) const
-{
-	UE::StructUtils::Private::ApplyChangesToPropertyDescs(
+	auto MoveProperty = [this, ChildPropertyHandle](const int32 Delta)
+	{
+		UE::StructUtils::Private::ApplyChangesToPropertyDescs(
 		LOCTEXT("OnPropertyMoved", "Move Property"), StructProperty, PropUtils,
 		[&ChildPropertyHandle, &Delta](TArray<FPropertyBagPropertyDesc>& PropertyDescs)
 		{
@@ -569,18 +569,13 @@ void FPropertyBagInstanceDataDetails::OnMoveProperty(TSharedPtr<IPropertyHandle>
 				}
 			}
 		});
-}
-
-TSharedRef<SWidget> FPropertyBagInstanceDataDetails::OnPropertyNameContent(TSharedPtr<IPropertyHandle> ChildPropertyHandle, TSharedPtr<SInlineEditableTextBlock> InlineWidget) const
-{
-	constexpr bool bInShouldCloseWindowAfterMenuSelection = true;
-	FMenuBuilder MenuBuilder(bInShouldCloseWindowAfterMenuSelection, nullptr);
-
+	};
+	
 	MenuBuilder.AddMenuEntry(
 		FUIAction(),
-		SNew(SPinTypeSelector, FGetPinTypeTree::CreateSP(this, &FPropertyBagInstanceDataDetails::GetFilteredVariableTypeTree))
-			.TargetPinType(this, &FPropertyBagInstanceDataDetails::OnGetPinInfo, ChildPropertyHandle)
-			.OnPinTypeChanged(this, &FPropertyBagInstanceDataDetails::PinInfoChanged, ChildPropertyHandle)
+		SNew(SPinTypeSelector, FGetPinTypeTree::CreateLambda(GetFilteredVariableTypeTree))
+			.TargetPinType_Lambda(GetPinInfo)
+			.OnPinTypeChanged_Lambda(PinInfoChanged)
 			.Schema(GetDefault<UEdGraphSchema_K2>())
 			.bAllowArrays(false)
 			.TypeTreeFilter(ETypeTreeFilter::None)
@@ -600,7 +595,7 @@ TSharedRef<SWidget> FPropertyBagInstanceDataDetails::OnPropertyNameContent(TShar
 		LOCTEXT("Remove", "Remove"),
 		LOCTEXT("Remove_ToolTip", "Remove property"),
 		FSlateIcon(FAppStyle::GetAppStyleSetName(), "Icons.Delete"),
-		FUIAction(FExecuteAction::CreateSP(this, &FPropertyBagInstanceDataDetails::OnRemoveProperty, ChildPropertyHandle))
+		FUIAction(FExecuteAction::CreateLambda(RemoveProperty))
 	);
 	
 	MenuBuilder.AddSeparator();
@@ -609,14 +604,14 @@ TSharedRef<SWidget> FPropertyBagInstanceDataDetails::OnPropertyNameContent(TShar
 		LOCTEXT("MoveUp", "Move Up"),
 		LOCTEXT("MoveUp_ToolTip", "Move property up in the list"),
 		FSlateIcon(FAppStyle::GetAppStyleSetName(), "Icons.ArrowUp"),
-		FUIAction(FExecuteAction::CreateSP(this, &FPropertyBagInstanceDataDetails::OnMoveProperty, ChildPropertyHandle, -1))
+		FUIAction(FExecuteAction::CreateLambda(MoveProperty, -1))
 	);
 
 	MenuBuilder.AddMenuEntry(
 		LOCTEXT("MoveDown", "Move Down"),
 		LOCTEXT("MoveDown_ToolTip", "Move property down in the list"),
 		FSlateIcon(FAppStyle::GetAppStyleSetName(), "Icons.ArrowDown"),
-		FUIAction(FExecuteAction::CreateSP(this, &FPropertyBagInstanceDataDetails::OnMoveProperty, ChildPropertyHandle, +1))
+		FUIAction(FExecuteAction::CreateLambda(MoveProperty, +1))
 	);
 
 	return MenuBuilder.MakeWidget();
diff --git a/Engine/Plugins/Experimental/StructUtils/Source/StructUtilsEditor/Public/PropertyBagDetails.h b/Engine/Plugins/Experimental/StructUtils/Source/StructUtilsEditor/Public/PropertyBagDetails.h
--- a/Engine/Plugins/Experimental/StructUtils/Source/StructUtilsEditor/Public/PropertyBagDetails.h
+++ b/Engine/Plugins/Experimental/StructUtils/Source/StructUtilsEditor/Public/PropertyBagDetails.h
@@ -9,6 +9,7 @@
 class IPropertyHandle;
 class IDetailPropertyRow;
 class SInlineEditableTextBlock;
+class SWidget;
 
 /**
  * Type customization for FInstancedPropertyBag.
@@ -56,16 +57,6 @@ public:
 	virtual void OnChildRowAdded(IDetailPropertyRow& ChildRow) override;
 
 protected:
-	FEdGraphPinType OnGetPinInfo(TSharedPtr<IPropertyHandle> ChildPropertyHandle) const;
-	void PinInfoChanged(const FEdGraphPinType& PinType, TSharedPtr<IPropertyHandle> ChildPropertyHandle) const;
-	
-	void GetFilteredVariableTypeTree(TArray<TSharedPtr<UEdGraphSchema_K2::FPinTypeTreeInfo>>& TypeTree, ETypeTreeFilter TypeTreeFilter) const;
-	FText GetPropertyName(TSharedPtr<IPropertyHandle> ChildPropertyHandle) const;
-	bool OnValidatePropertyName(const FText& InText, FText& OutErrorMessage, TSharedPtr<IPropertyHandle> ChildPropertyHandle) const;
-	void OnChangePropertyName(const FText& InNewText, ETextCommit::Type InCommitType, TSharedPtr<IPropertyHandle> ChildPropertyHandle) const;
-	void OnRemoveProperty(TSharedPtr<IPropertyHandle> ChildPropertyHandle) const;
-	void OnMoveProperty(TSharedPtr<IPropertyHandle> ChildPropertyHandle, int32 Delta) const;
-
 	TSharedRef<SWidget> OnPropertyNameContent(TSharedPtr<IPropertyHandle> ChildPropertyHandle, TSharedPtr<SInlineEditableTextBlock> InlineWidget) const;
 
 	TSharedPtr<IPropertyHandle> StructProperty;
