Property Bag changes for Remote Control Logic Plugin.

Adding functionality for creating a Property from source FProperty
Adding functionality for copy property value from source FProperty ContainerPtr

diff --git a/Engine/Plugins/Experimental/StructUtils/Source/StructUtils/Private/PropertyBag.cpp b/Engine/Plugins/Experimental/StructUtils/Source/StructUtils/Private/PropertyBag.cpp
--- a/Engine/Plugins/Experimental/StructUtils/Source/StructUtils/Private/PropertyBag.cpp
+++ b/Engine/Plugins/Experimental/StructUtils/Source/StructUtils/Private/PropertyBag.cpp
@@ -21,7 +21,11 @@ namespace UE::StructUtils::Private
 
 	uint64 CalcPropertyDescHash(const FPropertyBagPropertyDesc& Desc)
 	{
+#if WITH_EDITORONLY_DATA
+		const uint32 Hashes[] = { GetTypeHash(Desc.ID), GetTypeHash(Desc.Name), GetTypeHash(Desc.ValueType), GetTypeHash(Desc.MetaData) };
+#else
 		const uint32 Hashes[] = { GetTypeHash(Desc.ID), GetTypeHash(Desc.Name), GetTypeHash(Desc.ValueType) };
+#endif
 		return CityHash64WithSeed((const char*)Hashes, sizeof(Hashes), GetObjectHash(Desc.ValueTypeObject));
 	}
 
@@ -35,6 +39,109 @@ namespace UE::StructUtils::Private
 		return Hash;
 	}
 
+	EPropertyBagPropertyType GetValueTypeFromProperty(const FProperty* InSourceProperty)
+	{
+		if (CastField<FBoolProperty>(InSourceProperty))
+		{
+			return EPropertyBagPropertyType::Bool;
+		}
+		if (const FByteProperty* ByteProperty = CastField<FByteProperty>(InSourceProperty))
+		{
+			return ByteProperty->IsEnum() ? EPropertyBagPropertyType::Enum : EPropertyBagPropertyType::Byte;
+		}
+		if (CastField<FIntProperty>(InSourceProperty))
+		{
+			return EPropertyBagPropertyType::Int32;
+		}
+		if (CastField<FInt64Property>(InSourceProperty))
+		{
+			return EPropertyBagPropertyType::Int64;
+		}
+		if (CastField<FFloatProperty>(InSourceProperty))
+		{
+			return EPropertyBagPropertyType::Float;
+		}
+		if (CastField<FDoubleProperty>(InSourceProperty))
+		{
+			return EPropertyBagPropertyType::Double;
+		}
+		if (CastField<FNameProperty>(InSourceProperty))
+		{
+			return EPropertyBagPropertyType::Name;
+		}
+		if (CastField<FStrProperty>(InSourceProperty))
+		{
+			return EPropertyBagPropertyType::String;
+		}
+		if (CastField<FTextProperty>(InSourceProperty))
+		{
+			return EPropertyBagPropertyType::Text;
+		}
+		if (CastField<FEnumProperty>(InSourceProperty))
+		{
+			return EPropertyBagPropertyType::Enum;
+		}
+		if (CastField<FStructProperty>(InSourceProperty))
+		{
+			return EPropertyBagPropertyType::Struct;
+		}
+		if (CastField<FObjectProperty>(InSourceProperty))
+		{
+			return EPropertyBagPropertyType::Object;
+		}
+		if (CastField<FSoftObjectProperty>(InSourceProperty))
+		{
+			return EPropertyBagPropertyType::SoftObject;
+		}
+		if (CastField<FClassProperty>(InSourceProperty))
+		{
+			return EPropertyBagPropertyType::Class;
+		}
+		if (CastField<FSoftClassProperty>(InSourceProperty))
+		{
+			return EPropertyBagPropertyType::SoftClass;
+		}
+
+		return EPropertyBagPropertyType::None;
+	}
+
+	UObject* GetValueTypeObjectFromProperty(const FProperty* InSourceProperty)
+	{
+		if (const FByteProperty* ByteProperty = CastField<FByteProperty>(InSourceProperty))
+		{
+			if (ByteProperty->IsEnum())
+			{
+				return ByteProperty->Enum;
+			}
+		}
+		if (const FEnumProperty* EnumProp = CastField<FEnumProperty>(InSourceProperty))
+		{
+			return EnumProp->GetEnum();
+		}
+		if (const FStructProperty* StructProperty = CastField<FStructProperty>(InSourceProperty))
+		{
+			return StructProperty->Struct;
+		}
+		if (const FObjectProperty* ObjectProperty = CastField<FObjectProperty>(InSourceProperty))
+		{
+			return ObjectProperty->PropertyClass;
+		}
+		if (const FSoftObjectProperty* SoftObjectProperty = CastField<FSoftObjectProperty>(InSourceProperty))
+		{
+			return SoftObjectProperty->PropertyClass;
+		}
+		if (const FClassProperty* ClassProperty = CastField<FClassProperty>(InSourceProperty))
+		{
+			return ClassProperty->PropertyClass;
+		}
+		if (const FSoftClassProperty* SoftClassProperty = CastField<FSoftClassProperty>(InSourceProperty))
+		{
+			return SoftClassProperty->PropertyClass;
+		}
+
+		return nullptr;
+	}
+
 	FProperty* CreatePropertyFromDesc(const FPropertyBagPropertyDesc& Desc, UScriptStruct* PropertyScope)
 	{
 		switch (Desc.ValueType)
@@ -590,12 +697,55 @@ namespace UE::StructUtils::Private
 //  FPropertyBagPropertyDesc
 //----------------------------------------------------------------//
 
+void FPropertyBagPropertyDescMetaData::Serialize(FArchive& Ar)
+{
+	Ar << Key;
+	Ar << Value;
+}
+
+FPropertyBagPropertyDesc::FPropertyBagPropertyDesc(const FName InName, const FProperty* InSourceProperty)
+	: Name(InName)
+{
+	ValueType = UE::StructUtils::Private::GetValueTypeFromProperty(InSourceProperty);
+	ValueTypeObject = UE::StructUtils::Private::GetValueTypeObjectFromProperty(InSourceProperty);
+
+#if WITH_EDITORONLY_DATA
+	if (const TMap<FName, FString>* SourcePropertyMetaData = InSourceProperty->GetMetaDataMap())
+	{
+		for (const TPair<FName, FString>& MetaDataPair : *SourcePropertyMetaData)
+		{
+			MetaData.Add({ MetaDataPair.Key, MetaDataPair.Value });
+		}
+	}
+#endif
+}
+
 static FArchive& operator<<(FArchive& Ar, FPropertyBagPropertyDesc& Bag)
 {
 	Ar << Bag.ValueTypeObject;
 	Ar << Bag.ID;
 	Ar << Bag.Name;
 	Ar << Bag.ValueType;
+
+	bool bHasMetaData = false;
+#if WITH_EDITORONLY_DATA
+	if (Ar.IsSaving())
+	{
+		bHasMetaData = !Ar.IsCooking() && Bag.MetaData.Num() > 0;
+	}
+#endif
+	Ar << bHasMetaData;
+
+	if (bHasMetaData)
+	{
+#if WITH_EDITORONLY_DATA
+		Ar << Bag.MetaData;
+#else
+		TArray<FPropertyBagPropertyDescMetaData> TempMetaData; 
+		Ar << TempMetaData;
+#endif
+	}	
+	
 	return Ar;
 }
 
@@ -682,6 +832,16 @@ void FInstancedPropertyBag::CopyMatchingValuesByID(const FInstancedPropertyBag&
 	UE::StructUtils::Private::CopyMatchingValuesByID(Other.Value, Value);
 }
 
+int32 FInstancedPropertyBag::GetNumPropertiesInBag() const
+{
+	if (const UPropertyBag* BagStruct = GetPropertyBagStruct())
+	{
+		return BagStruct->PropertyDescs.Num();
+	}
+
+	return 0;
+}
+
 void FInstancedPropertyBag::AddProperties(const TConstArrayView<FPropertyBagPropertyDesc> NewDescs)
 {
 	TArray<FPropertyBagPropertyDesc> Descs;
@@ -713,6 +873,11 @@ void FInstancedPropertyBag::AddProperty(const FName InName, const EPropertyBagPr
 	AddProperties({ FPropertyBagPropertyDesc(InName, InValueType, InValueTypeObject) });
 }
 
+void FInstancedPropertyBag::AddProperty(const FName InName, const FProperty* InSourceProperty)
+{
+	AddProperties({ FPropertyBagPropertyDesc(InName, InSourceProperty ) });
+}
+
 void FInstancedPropertyBag::RemovePropertiesByName(const TConstArrayView<FName> PropertiesToRemove)
 {
 	TArray<FPropertyBagPropertyDesc> Descs;
@@ -1062,6 +1227,30 @@ EPropertyBagResult FInstancedPropertyBag::SetValueClass(const FName Name, UClass
 	return EPropertyBagResult::Success;
 }
 
+EPropertyBagResult FInstancedPropertyBag::SetValue(const FName Name, const FProperty* InSourceProperty, const void* InSourceContainerAddress) const
+{
+	const FPropertyBagPropertyDesc* Desc = FindPropertyDescByName(Name);
+	if (Desc == nullptr || InSourceProperty == nullptr || InSourceContainerAddress == nullptr)
+	{
+		return EPropertyBagResult::PropertyNotFound;
+	}
+	check(Desc->CachedProperty);
+
+	void* TargetAddress = Value.GetMutableMemory() + Desc->CachedProperty->GetOffset_ForInternal();
+	void const* SourceAddress = InSourceProperty->ContainerPtrToValuePtr<void>(InSourceContainerAddress);
+
+	if (InSourceProperty->GetClass() == Desc->CachedProperty->GetClass())
+	{
+		Desc->CachedProperty->CopyCompleteValue(TargetAddress, SourceAddress);
+	}
+	else
+	{
+		return EPropertyBagResult::TypeMismatch;
+	}
+
+	return EPropertyBagResult::Success;
+}
+
 bool FInstancedPropertyBag::Serialize(FArchive& Ar)
 {
 	enum class EVersion : uint8
@@ -1182,9 +1371,18 @@ const UPropertyBag* UPropertyBag::GetOrCreateFromDescs(const TConstArrayView<FPr
 		{
 			Desc.ID = FGuid::NewGuid();
 		}
-		
+
+
 		if (FProperty* NewProperty = UE::StructUtils::Private::CreatePropertyFromDesc(Desc, NewBag))
 		{
+#if WITH_EDITORONLY_DATA
+			// Add metadata
+			for (const FPropertyBagPropertyDescMetaData& PropertyDescMetaData : Desc.MetaData)
+			{
+				NewProperty->SetMetaData(*PropertyDescMetaData.Key.ToString(), *PropertyDescMetaData.Value);
+			}
+#endif
+			
 			NewProperty->SetPropertyFlags(CPF_Edit);
 			NewBag->AddCppProperty(NewProperty);
 			Desc.CachedProperty = NewProperty; 
diff --git a/Engine/Plugins/Experimental/StructUtils/Source/StructUtils/Public/PropertyBag.h b/Engine/Plugins/Experimental/StructUtils/Source/StructUtils/Public/PropertyBag.h
--- a/Engine/Plugins/Experimental/StructUtils/Source/StructUtils/Public/PropertyBag.h
+++ b/Engine/Plugins/Experimental/StructUtils/Source/StructUtils/Public/PropertyBag.h
@@ -38,6 +38,53 @@ enum class EPropertyBagResult : uint8
 	PropertyNotFound,	// Could not find property of specified name.
 };
 
+USTRUCT()
+struct STRUCTUTILS_API FPropertyBagPropertyDescMetaData
+{
+	GENERATED_BODY()
+	
+	FPropertyBagPropertyDescMetaData() = default;
+	FPropertyBagPropertyDescMetaData(const FName InKey, const FString& InValue)
+		: Key(InKey)
+		, Value(InValue)
+	{
+	}
+	
+	UPROPERTY()
+	FName Key;
+
+	UPROPERTY()
+	FString Value;
+
+	void Serialize(FArchive& Ar);
+
+	FORCEINLINE friend FArchive& operator<<(FArchive& Ar, FPropertyBagPropertyDescMetaData& PropertyDescMetaData)
+	{
+		PropertyDescMetaData.Serialize(Ar);
+		return Ar;
+	}
+
+	friend FORCEINLINE uint32 GetTypeHash(const FPropertyBagPropertyDescMetaData& PropertyDescMetaData)
+	{
+		return HashCombine(GetTypeHash(PropertyDescMetaData.Key), GetTypeHash(PropertyDescMetaData.Value));
+	}
+
+	friend FORCEINLINE uint32 GetTypeHash(const TArrayView<const FPropertyBagPropertyDescMetaData>& MetaData)
+	{
+		uint32 Hash = GetTypeHash(MetaData.Num());
+		for (const FPropertyBagPropertyDescMetaData& PropertyDescMetaData : MetaData)
+		{
+			Hash = HashCombine(Hash, GetTypeHash(PropertyDescMetaData));
+		}
+		return Hash;
+	}
+
+	friend FORCEINLINE uint32 GetTypeHash(const TArray<FPropertyBagPropertyDescMetaData>& MetaData)
+	{
+		return GetTypeHash(TArrayView<const FPropertyBagPropertyDescMetaData>(MetaData.GetData(), MetaData.Num()));
+	}
+};
+
 /** Describes a property in the property bag. */
 USTRUCT()
 struct STRUCTUTILS_API FPropertyBagPropertyDesc
@@ -45,6 +92,7 @@ struct STRUCTUTILS_API FPropertyBagPropertyDesc
 	GENERATED_BODY()
 
 	FPropertyBagPropertyDesc() = default;
+	FPropertyBagPropertyDesc(const FName InName, const FProperty* InSourceProperty);
 	FPropertyBagPropertyDesc(const FName InName, const EPropertyBagPropertyType InValueType, UObject* InValueTypeObject = nullptr)
 		: ValueTypeObject(InValueTypeObject)
 		, Name(InName)
@@ -67,7 +115,6 @@ struct STRUCTUTILS_API FPropertyBagPropertyDesc
 	/** @return true if the property type is class or soft class */
 	bool IsClassType() const;
 
-	
 	/** Pointer to object that defines the Enum, Struct, or Class. */
 	UPROPERTY(EditAnywhere, Category="Default")
 	TObjectPtr<const UObject> ValueTypeObject = nullptr;
@@ -84,6 +131,12 @@ struct STRUCTUTILS_API FPropertyBagPropertyDesc
 	UPROPERTY(EditAnywhere, Category="Default")
 	EPropertyBagPropertyType ValueType = EPropertyBagPropertyType::None;
 
+#if WITH_EDITORONLY_DATA
+	/** Editor-only meta data for CachedProperty */
+	UPROPERTY(EditAnywhere, Category="Default")
+	TArray<FPropertyBagPropertyDescMetaData> MetaData;
+#endif
+
 	/** Cached property pointer, set in UPropertyBag::GetOrCreateFromDescs. */
 	const FProperty* CachedProperty = nullptr;
 };
@@ -166,6 +219,9 @@ struct STRUCTUTILS_API FInstancedPropertyBag
 	 */
 	void CopyMatchingValuesByID(const FInstancedPropertyBag& NewDescs) const;
 
+	/** Returns number of the Properties in this Property Bag */
+	int32 GetNumPropertiesInBag() const;
+
 	/**
 	 * Adds properties to the bag. If property of same name already exists, it will be replaced with the new type.
 	 * Numeric property values will be converted if possible, when a property's type changes.
@@ -180,6 +236,12 @@ struct STRUCTUTILS_API FInstancedPropertyBag
 	 */
 	void AddProperty(const FName InName, const EPropertyBagPropertyType InValueType, UObject* InValueTypeObject = nullptr);
 
+	/**
+	 * Adds a new property to the bag. Property type duplicated from source property to. If property of same name already exists, it will be replaced with the new type.
+	 * @param Descs Descriptors of new properties to add.
+	 */
+	void AddProperty(const FName InName, const FProperty* InSourceProperty);
+
 	/**
 	 * Removes properties from the bag by name if they exists.
 	 */
@@ -331,6 +393,12 @@ struct STRUCTUTILS_API FInstancedPropertyBag
 		return SetValueObject(Name, (UObject*)InValue);
 	}
 
+	/**
+	 * Sets property value from given source property and source container address
+	 * A property must exists in that bag before it can be set. 
+	 */
+	EPropertyBagResult SetValue(const FName Name, const FProperty* InSourceProperty, const void* InSourceContainerAddress) const;
+
 	bool Serialize(FArchive& Ar);
 
 protected:
