PropertyBag: Fix potential null ptr access
- fixed nullptr access reported by static analysis
- improved error handling in case of missing struct/class/enum type

UE-150958

diff --git a/Engine/Plugins/Experimental/StructUtils/Source/StructUtils/Private/PropertyBag.cpp b/Engine/Plugins/Experimental/StructUtils/Source/StructUtils/Private/PropertyBag.cpp
--- a/Engine/Plugins/Experimental/StructUtils/Source/StructUtils/Private/PropertyBag.cpp
+++ b/Engine/Plugins/Experimental/StructUtils/Source/StructUtils/Private/PropertyBag.cpp
@@ -1256,6 +1256,7 @@ bool FInstancedPropertyBag::Serialize(FArchive& Ar)
 	enum class EVersion : uint8
 	{
 		InitialVersion = 0,
+		SerializeStructSize,
 		// -----<new versions can be added above this line>-----
 		VersionPlusOne,
 		LatestVersion = VersionPlusOne - 1
@@ -1276,23 +1277,57 @@ bool FInstancedPropertyBag::Serialize(FArchive& Ar)
 		if (Ar.IsLoading())
 		{
 			TArray<FPropertyBagPropertyDesc> PropertyDescs;
-
 			Ar << PropertyDescs;
 
 			BagStruct = const_cast<UPropertyBag*>(UPropertyBag::GetOrCreateFromDescs(PropertyDescs));
 			Value.InitializeAs(BagStruct);
+
+			// Size of the serialized memory
+			int32 SerialSize = 0; 
+			if (Version >= EVersion::SerializeStructSize)
+			{
+				Ar << SerialSize;
+			}
+			
+			// BagStruct can be null if it contains structs, classes or enums that could not be found.
+			if (BagStruct != nullptr && Value.GetMutableMemory() != nullptr)
+			{
+				BagStruct->SerializeItem(Ar, Value.GetMutableMemory(), /*Defaults*/nullptr);
+			}
+			else
+			{
+				UE_LOG(LogCore, Warning, TEXT("Unable to create serialized UPropertyBag -> Advance %u bytes in the archive and reset to empty FInstancedPropertyBag"), SerialSize);
+				Ar.Seek(Ar.Tell() + SerialSize);
+			}
 		}
-		if (Ar.IsSaving())
+		else if (Ar.IsSaving())
 		{
+			check(BagStruct);
 			Ar << BagStruct->PropertyDescs;
-		}
 
-		if (ensureMsgf(Value.GetMutableMemory() != nullptr, TEXT("A valid script struct should always have allocated memory")))
-		{
-			BagStruct->SerializeItem(Ar, Value.GetMutableMemory(), /*Defaults*/nullptr);
+			const int64 SizeOffset = Ar.Tell(); // Position to write the actual size after struct serialization
+			int32 SerialSize = 0;
+			// Size of the serialized memory (reserve location)
+			Ar << SerialSize;
+
+			const int64 InitialOffset = Ar.Tell(); // Position before struct serialization to compute its serial size
+
+			// BagStruct can be null if it contains structs, classes or enums that are not present anymore. 
+			if (BagStruct != nullptr && Value.GetMutableMemory() != nullptr)
+			{
+				BagStruct->SerializeItem(Ar, Value.GetMutableMemory(), /*Defaults*/nullptr);
+			}
+		
+			const int64 FinalOffset = Ar.Tell(); // Keep current offset to reset the archive pos after write the serial size
+
+			// Size of the serialized memory
+			Ar.Seek(SizeOffset);	// Go back in the archive to write the actual size
+			SerialSize = (int32)(FinalOffset - InitialOffset);
+			Ar << SerialSize;
+			Ar.Seek(FinalOffset);	// Reset archive to its position
 		}
 	}
-
+	
 	return true;
 }
 
