Fixing EditorStyle->AppStyle merge errors

diff --git a/Engine/Plugins/Experimental/StructUtils/Source/StructUtilsEditor/Private/Customizations/PropertyBagDetails.cpp b/Engine/Plugins/Experimental/StructUtils/Source/StructUtilsEditor/Private/Customizations/PropertyBagDetails.cpp
--- a/Engine/Plugins/Experimental/StructUtils/Source/StructUtilsEditor/Private/Customizations/PropertyBagDetails.cpp
+++ b/Engine/Plugins/Experimental/StructUtils/Source/StructUtilsEditor/Private/Customizations/PropertyBagDetails.cpp
@@ -1,4 +1,4 @@
-﻿// Copyright Epic Games, Inc. All Rights Reserved.
+// Copyright Epic Games, Inc. All Rights Reserved.
 
 #include "PropertyBagDetails.h"
 #include "InstancedStructDetails.h"
@@ -18,6 +18,7 @@
 #include "Kismet2/BlueprintEditorUtils.h"
 #include "SPinTypeSelector.h"
 #include "PropertyBag.h"
+#include "Styling/AppStyle.h"
 
 #define LOCTEXT_NAMESPACE "StructUtilsEditor"
 
@@ -459,7 +460,7 @@ void FPropertyBagInstanceDataDetails::OnChildRowAdded(IDetailPropertyRow& ChildR
 		NameWidget = SNew(SComboButton)
 			.OnGetMenuContent(this, &FPropertyBagInstanceDataDetails::OnPropertyNameContent, ChildRow.GetPropertyHandle(), InlineWidget)
 			.ContentPadding(2)
-			.ButtonStyle(FEditorStyle::Get(), "HoverHintOnly")
+			.ButtonStyle(FAppStyle::Get(), "HoverHintOnly")
 			.ButtonContent()
 			[
 				InlineWidget.ToSharedRef()
@@ -676,7 +677,7 @@ TSharedPtr<SWidget> FPropertyBagDetails::MakeAddPropertyWidget(TSharedPtr<IPrope
 			SNew(SButton)
 			.VAlign(VAlign_Center)
 			.HAlign(HAlign_Center)
-			.ButtonStyle(FEditorStyle::Get(), "HoverHintOnly")
+			.ButtonStyle(FAppStyle::Get(), "HoverHintOnly")
 			.ToolTipText(LOCTEXT("AddProperty_Tooltip", "Add new property"))
 			.OnClicked_Lambda([InStructProperty, InPropUtils]()
 			{
@@ -706,7 +707,7 @@ TSharedPtr<SWidget> FPropertyBagDetails::MakeAddPropertyWidget(TSharedPtr<IPrope
 			.Content()
 			[
 				SNew(SImage)
-				.Image(FEditorStyle::GetBrush("Icons.PlusCircle"))
+				.Image(FAppStyle::GetBrush("Icons.PlusCircle"))
 			]
 		];
 }
diff --git a/Engine/Plugins/Runtime/StateTree/Source/StateTreeEditorModule/Private/SStateTreeViewRow.cpp b/Engine/Plugins/Runtime/StateTree/Source/StateTreeEditorModule/Private/SStateTreeViewRow.cpp
--- a/Engine/Plugins/Runtime/StateTree/Source/StateTreeEditorModule/Private/SStateTreeViewRow.cpp
+++ b/Engine/Plugins/Runtime/StateTree/Source/StateTreeEditorModule/Private/SStateTreeViewRow.cpp
@@ -231,7 +231,7 @@ void SStateTreeViewRow::Construct(const FArguments& InArgs, const TSharedRef<STa
 			.Visibility(this, &SStateTreeViewRow::GetLinkedStateVisibility)
 			[
 				SNew(SBorder)
-				.BorderImage(FEditorStyle::GetBrush("WhiteBrush"))
+				.BorderImage(FAppStyle::GetBrush("WhiteBrush"))
 				.BorderBackgroundColor(LinkBackground)
 				.Padding(FMargin(12.0f, 0.0f, 16.0f, 0.0f))
 				[
diff --git a/Engine/Plugins/Runtime/StateTree/Source/StateTreeEditorModule/Private/StateTreeEditor.cpp b/Engine/Plugins/Runtime/StateTree/Source/StateTreeEditorModule/Private/StateTreeEditor.cpp
--- a/Engine/Plugins/Runtime/StateTree/Source/StateTreeEditorModule/Private/StateTreeEditor.cpp
+++ b/Engine/Plugins/Runtime/StateTree/Source/StateTreeEditorModule/Private/StateTreeEditor.cpp
@@ -754,20 +754,20 @@ FSlateIcon FStateTreeEditor::GetCompileStatusImage() const
 
 	if (StateTree == nullptr)
 	{
-		return FSlateIcon(FEditorStyle::GetStyleSetName(), CompileStatusBackground, NAME_None, CompileStatusUnknown);
+		return FSlateIcon(FAppStyle::GetAppStyleSetName(), CompileStatusBackground, NAME_None, CompileStatusUnknown);
 	}
 	
 	if (!StateTree->IsReadyToRun())
 	{
-		return FSlateIcon(FEditorStyle::GetStyleSetName(), CompileStatusBackground, NAME_None, CompileStatusError);
+		return FSlateIcon(FAppStyle::GetAppStyleSetName(), CompileStatusBackground, NAME_None, CompileStatusError);
 	}
 	
 	if (StateTree->LastCompiledEditorDataHash != EditorDataHash)
 	{
-		return FSlateIcon(FEditorStyle::GetStyleSetName(), CompileStatusBackground, NAME_None, CompileStatusUnknown);
+		return FSlateIcon(FAppStyle::GetAppStyleSetName(), CompileStatusBackground, NAME_None, CompileStatusUnknown);
 	}
 	
-	return FSlateIcon(FEditorStyle::GetStyleSetName(), CompileStatusBackground, NAME_None, CompileStatusGood);
+	return FSlateIcon(FAppStyle::GetAppStyleSetName(), CompileStatusBackground, NAME_None, CompileStatusGood);
 }
 
 void FStateTreeEditor::UpdateAsset()
