Added an "index iterator" to StructTypeBitSet allowing retrieval of iterrative retrieval of indices representing a given value.
Also, added some unit tests for it.

diff --git a/Engine/Plugins/Experimental/StructUtils/Source/StructUtils/Public/StructTypeBitSet.h b/Engine/Plugins/Experimental/StructUtils/Source/StructUtils/Public/StructTypeBitSet.h
--- a/Engine/Plugins/Experimental/StructUtils/Source/StructUtils/Public/StructTypeBitSet.h
+++ b/Engine/Plugins/Experimental/StructUtils/Source/StructUtils/Public/StructTypeBitSet.h
@@ -281,6 +281,45 @@ public:
 		}
 	}
 
+	/** Retrieves indices of bits of given value. */
+	struct FIndexIterator
+	{
+	public:
+		explicit FIndexIterator(const FBitArrayExt& BitArray, const bool bInValueToCheck = true)
+			: It(BitArray), bValueToCheck(bInValueToCheck)
+		{
+			if (It && It.GetValue() != bInValueToCheck)
+			{
+				// will result in either setting IT to the first bInValue, or making bool(It) == false
+				++(*this);
+			}
+		}
+
+		operator bool() const { return bool(It); }
+
+		/** @todo can be optimized to skip whole empty words, similar to TBitArray<>::Find */
+		FIndexIterator& operator++() 
+		{ 
+			while (++It)
+			{
+				if (It.GetValue() == bValueToCheck)
+				{
+					break;
+				}
+			}
+			return *this; 
+		}
+
+		int32 operator*() const
+		{
+			return It.GetIndex();
+		}
+
+	private:
+		TBitArray<>::FConstIterator It;
+		const bool bValueToCheck;
+	};
+
 private:
 	/** 
 	 * A private constructor for a creating an instance straight from TBitArrays. 
@@ -482,6 +521,11 @@ public:
 		return StructTypesBitArray.Contains(BitIndex);
 	}
 
+	FIndexIterator GetIndexIterator(const bool bValueToCheck = true) const
+	{
+		return FIndexIterator(StructTypesBitArray, bValueToCheck);
+	}
+
 	static int32 GetMaxNum() 
 	{		
 		return TStructTrackerWrapper::StructTracker.Num();
diff --git a/Engine/Plugins/Experimental/StructUtils/Source/StructUtilsTestSuite/Private/ScriptStructTypeBitSetTest.cpp b/Engine/Plugins/Experimental/StructUtils/Source/StructUtilsTestSuite/Private/ScriptStructTypeBitSetTest.cpp
--- a/Engine/Plugins/Experimental/StructUtils/Source/StructUtilsTestSuite/Private/ScriptStructTypeBitSetTest.cpp
+++ b/Engine/Plugins/Experimental/StructUtils/Source/StructUtilsTestSuite/Private/ScriptStructTypeBitSetTest.cpp
@@ -327,6 +327,82 @@ struct FStructUtilsTest_CompileTimeBitSetReversed : FAITestBase
 
 IMPLEMENT_AI_INSTANT_TEST(FStructUtilsTest_CompileTimeBitSetReversed, "System.StructUtils.BitSet.CompileTime_Reversed");
 
+struct FStructUtilsTest_BitSetIndexIteratorBase : FAITestBase
+{
+	TArray<int32> Indices;
+	FTestStructBitSet BitSet;
+	bool bValue = true;
+
+	FStructUtilsTest_BitSetIndexIteratorBase()
+		: Indices({ 0, 3, 9, 230 })
+	{}
+
+	virtual bool SetUp() override
+	{
+		for (int32 Index : Indices)
+		{
+			BitSet.AddBit(Index);
+		}
+		return FAITestBase::SetUp();
+	}
+
+	virtual bool InstantTest() override
+	{
+		int32 i = 0;
+		for (auto It = BitSet.GetIndexIterator(/*bValue=*/bValue); It; ++It, ++i)
+		{
+			AITEST_EQUAL("Indices fetched by FIndexIterator need to reflect Bitset\'s contents", Indices[i], *It);
+		}
+		return true;
+	}
+};
+
+IMPLEMENT_AI_INSTANT_TEST(FStructUtilsTest_BitSetIndexIteratorBase, "System.StructUtils.BitSet.IndexIterator");
+
+struct FStructUtilsTest_BitSetIndexIteratorEmpty : FStructUtilsTest_BitSetIndexIteratorBase
+{
+	FStructUtilsTest_BitSetIndexIteratorEmpty()
+	{
+		Indices = {};
+	}
+};
+IMPLEMENT_AI_INSTANT_TEST(FStructUtilsTest_BitSetIndexIteratorEmpty, "System.StructUtils.BitSet.IndexIterator.Empty");
+
+struct FStructUtilsTest_BitSetIndexIteratorTrailingFalse : FStructUtilsTest_BitSetIndexIteratorBase
+{
+	virtual bool SetUp() override
+	{
+		Indices.Sort();
+		for (int32 Index : Indices)
+		{
+			BitSet.AddBit(Index);
+		}
+		const int32 TrailingFalseBitIndex = Indices.Last() + 1;
+		BitSet.AddBit(TrailingFalseBitIndex);
+		BitSet.RemoveBit(TrailingFalseBitIndex);
+
+		return FAITestBase::SetUp();
+	}
+};
+IMPLEMENT_AI_INSTANT_TEST(FStructUtilsTest_BitSetIndexIteratorTrailingFalse, "System.StructUtils.BitSet.IndexIterator.TrailingFalse");
+
+struct FStructUtilsTest_BitSetIndexIteratorBeginningFalse : FStructUtilsTest_BitSetIndexIteratorBase
+{
+	virtual bool SetUp() override
+	{
+		Indices.Sort();
+		for (int32 Index : Indices)
+		{
+			BitSet.AddBit(Index);
+		}
+		BitSet.RemoveBit(Indices[0]);
+		Indices.RemoveAt(0);
+
+		return FAITestBase::SetUp();
+	}
+};
+IMPLEMENT_AI_INSTANT_TEST(FStructUtilsTest_BitSetIndexIteratorBeginningFalse, "System.StructUtils.BitSet.IndexIterator.BeginningFalse");
+
 } // namespace FStructTypeBitSetTests
 
 #undef LOCTEXT_NAMESPACE
