PropertyBag: allow to load bags with missing types.
- Allow to load property bags with missing struct, enum, and object types
- Mark missing types with specific types so that we can report error in the UI
- Report missing types in the UI

diff --git a/Engine/Plugins/Experimental/StructUtils/Source/StructUtils/Private/PropertyBag.cpp b/Engine/Plugins/Experimental/StructUtils/Source/StructUtils/Private/PropertyBag.cpp
--- a/Engine/Plugins/Experimental/StructUtils/Source/StructUtils/Private/PropertyBag.cpp
+++ b/Engine/Plugins/Experimental/StructUtils/Source/StructUtils/Private/PropertyBag.cpp
@@ -1345,52 +1345,47 @@ const UPropertyBag* UPropertyBag::GetOrCreateFromDescs(const TConstArrayView<FPr
 		return ExistingBag;
 	}
 
-	// Validate descriptors
-	bool bAllValid = true;
-	for (const FPropertyBagPropertyDesc& Desc : PropertyDescs)
+	UPropertyBag* NewBag = NewObject<UPropertyBag>(GetTransientPackage(), *ScriptStructName, RF_Standalone | RF_Transient);
+
+	NewBag->PropertyDescs = PropertyDescs;
+
+	// Fix missing structs, enums, and objects.
+	for (FPropertyBagPropertyDesc& Desc : NewBag->PropertyDescs)
 	{
 		if (Desc.ValueType == EPropertyBagPropertyType::Struct)
 		{
 			if (Desc.ValueTypeObject == nullptr || Desc.ValueTypeObject->GetClass()->IsChildOf(UScriptStruct::StaticClass()) == false)
 			{
-				ensureAlwaysMsgf(false, TEXT("Struct property must have valid struct type object."));
-				bAllValid = false;
+				UE_LOG(LogCore, Warning, TEXT("PropertyBag: Struct property '%s' is missing type."), *Desc.Name.ToString());
+				Desc.ValueTypeObject = FPropertyBagMissingStruct::StaticStruct();
 			}
 		}
 		else if (Desc.ValueType == EPropertyBagPropertyType::Enum)
 		{
 			if (Desc.ValueTypeObject == nullptr || Desc.ValueTypeObject->GetClass()->IsChildOf(UEnum::StaticClass()) == false)
 			{
-				ensureAlwaysMsgf(false, TEXT("Enum property must have valid enum type object."));
-				bAllValid = false;
+				UE_LOG(LogCore, Warning, TEXT("PropertyBag: Enum property '%s' is missing type."), *Desc.Name.ToString());
+				Desc.ValueTypeObject = StaticEnum<EPropertyBagMissingEnum>();
 			}
 		}
 		else if (Desc.ValueType == EPropertyBagPropertyType::Object || Desc.ValueType == EPropertyBagPropertyType::SoftObject)
 		{
 			if (Desc.ValueTypeObject == nullptr)
 			{
-				ensureAlwaysMsgf(false, TEXT("Object property must have valid object type object."));
-				bAllValid = false;
+				UE_LOG(LogCore, Warning, TEXT("PropertyBag: Object property '%s' is missing type."), *Desc.Name.ToString());
+				Desc.ValueTypeObject = UPropertyBagMissingObject::StaticClass();
 			}
 		}
 		else if (Desc.ValueType == EPropertyBagPropertyType::Class || Desc.ValueType == EPropertyBagPropertyType::SoftClass)
 		{
 			if (Desc.ValueTypeObject == nullptr || Desc.ValueTypeObject->GetClass()->IsChildOf(UClass::StaticClass()) == false)
 			{
-				ensureAlwaysMsgf(false, TEXT("Class property must have valid class type object."));
-				bAllValid = false;
+				UE_LOG(LogCore, Warning, TEXT("PropertyBag: Class property '%s' is missing type."), *Desc.Name.ToString());
+				Desc.ValueTypeObject = UPropertyBagMissingObject::StaticClass();
 			}
 		}
 	}
-	if (!bAllValid)
-	{
-		return nullptr;
-	}
 	
-	UPropertyBag* NewBag = NewObject<UPropertyBag>(GetTransientPackage(), *ScriptStructName, RF_Standalone | RF_Transient);
-
-	NewBag->PropertyDescs = PropertyDescs;
-
 	// Remove properties with same name
 	for (int32 Index = 0; Index < NewBag->PropertyDescs.Num() - 1; Index++)
 	{
diff --git a/Engine/Plugins/Experimental/StructUtils/Source/StructUtils/Public/PropertyBag.h b/Engine/Plugins/Experimental/StructUtils/Source/StructUtils/Public/PropertyBag.h
--- a/Engine/Plugins/Experimental/StructUtils/Source/StructUtils/Public/PropertyBag.h
+++ b/Engine/Plugins/Experimental/StructUtils/Source/StructUtils/Public/PropertyBag.h
@@ -414,6 +414,28 @@ template<> struct TStructOpsTypeTraits<FInstancedPropertyBag> : public TStructOp
 	};
 };
 
+/**
+ * Dummy types used to mark up missing types when creating property bags. These are used in the UI to display error message.
+ */
+UENUM()
+enum class EPropertyBagMissingEnum : uint8
+{
+	Missing,
+};
+
+USTRUCT()
+struct STRUCTUTILS_API FPropertyBagMissingStruct
+{
+	GENERATED_BODY()
+};
+
+UCLASS()
+class STRUCTUTILS_API UPropertyBagMissingObject : public UObject
+{
+	GENERATED_BODY()
+};
+
+
 /**
  * A script struct that is used to store the value of the property bag instance.
  * References to UPropertyBag cannot be serialized, instead the array of the properties
diff --git a/Engine/Plugins/Experimental/StructUtils/Source/StructUtilsEditor/Private/Customizations/PropertyBagDetails.cpp b/Engine/Plugins/Experimental/StructUtils/Source/StructUtilsEditor/Private/Customizations/PropertyBagDetails.cpp
--- a/Engine/Plugins/Experimental/StructUtils/Source/StructUtilsEditor/Private/Customizations/PropertyBagDetails.cpp
+++ b/Engine/Plugins/Experimental/StructUtils/Source/StructUtilsEditor/Private/Customizations/PropertyBagDetails.cpp
@@ -40,6 +40,33 @@ bool IsScriptStruct(const TSharedPtr<IPropertyHandle>& PropertyHandle)
 	return StructProperty && StructProperty->Struct == TBaseStructure<T>::Get();
 }
 
+/** @return true if the property is one of the known missing types. */
+bool HasMissingType(const TSharedPtr<IPropertyHandle>& PropertyHandle)
+{
+	if (!PropertyHandle)
+	{
+		return false;
+	}
+
+	// Handles Struct
+	if (FStructProperty* StructProperty = CastField<FStructProperty>(PropertyHandle->GetProperty()))
+	{
+		return StructProperty->Struct == FPropertyBagMissingStruct::StaticStruct();
+	}
+	// Handles Object, SoftObject, Class, SoftClass.
+	if (FObjectPropertyBase* ObjectProperty = CastField<FObjectPropertyBase>(PropertyHandle->GetProperty()))
+	{
+		return ObjectProperty->PropertyClass == UPropertyBagMissingObject::StaticClass();
+	}
+	// Handles Enum
+	if (FEnumProperty* EnumProperty = CastField<FEnumProperty>(PropertyHandle->GetProperty()))
+	{
+		return EnumProperty->GetEnum() == StaticEnum<EPropertyBagMissingEnum>();
+	}
+
+	return false;
+}
+
 /** @return property bag struct common to all edited properties. */
 const UPropertyBag* GetCommonBagStruct(TSharedPtr<IPropertyHandle> StructProperty)
 {
@@ -410,10 +437,10 @@ void FPropertyBagInstanceDataDetails::OnChildRowAdded(IDetailPropertyRow& ChildR
 	FDetailWidgetRow Row;
 	ChildRow.GetDefaultWidgets(NameWidget, ValueWidget, Row);
 
+	TSharedPtr<IPropertyHandle> ChildPropertyHandle = ChildRow.GetPropertyHandle();
+	
 	if (!bFixedLayout)
 	{
-		TSharedPtr<IPropertyHandle> ChildPropertyHandle = ChildRow.GetPropertyHandle();
-
 		// Inline editable name
 		TSharedPtr<SInlineEditableTextBlock> InlineWidget = SNew(SInlineEditableTextBlock)
 			.Font(IDetailLayoutBuilder::GetDetailFont())
@@ -466,12 +493,35 @@ void FPropertyBagInstanceDataDetails::OnChildRowAdded(IDetailPropertyRow& ChildR
 				InlineWidget.ToSharedRef()
 			];
 	}
-		
+
 	ChildRow
 		.CustomWidget(/*bShowChildren*/true)
 		.NameContent()
 		[
-			NameWidget.ToSharedRef()
+			SNew(SHorizontalBox)
+			// Error icon
+			+ SHorizontalBox::Slot()
+			.AutoWidth()
+			.HAlign(HAlign_Right)
+			.VAlign(VAlign_Center)
+			.Padding(0,0,2,0)
+			[
+				SNew(SBox)
+				.WidthOverride(12)
+				.HeightOverride(12)
+				[
+					SNew(SImage)
+					.ToolTipText(LOCTEXT("MissingType", "The property is missing type. The Struct, Enum, or Object may have been removed."))
+					.Visibility_Lambda([this, ChildPropertyHandle]() { return UE::StructUtils::Private::HasMissingType(ChildPropertyHandle) ? EVisibility::Visible : EVisibility::Collapsed; })
+					.Image(FAppStyle::GetBrush("Icons.Error"))
+				]
+			]
+			// Name
+			+SHorizontalBox::Slot()
+			.AutoWidth()
+			[
+				NameWidget.ToSharedRef()
+			]
 		]
 		.ValueContent()
 		[
