FInstancedStructArray, heterogenous array of structs

diff --git a/Engine/Plugins/Experimental/StructUtils/Source/StructUtils/Private/InstancedStructArray.cpp b/Engine/Plugins/Experimental/StructUtils/Source/StructUtils/Private/InstancedStructArray.cpp
new file mode 100644
--- /dev/null
+++ b/Engine/Plugins/Experimental/StructUtils/Source/StructUtils/Private/InstancedStructArray.cpp
@@ -0,0 +1,555 @@
+﻿// Copyright Epic Games, Inc. All Rights Reserved.
+#include "InstancedStructArray.h"
+#include "Serialization/PropertyLocalizationDataGathering.h"
+
+// From InstancedStruct.cpp to support localization.
+namespace UE::StructUtils::Private
+{
+#if WITH_EDITORONLY_DATA
+	void GatherForLocalization(const FString& PathToParent, const UScriptStruct* Struct, const void* StructData, const void* DefaultStructData,
+								FPropertyLocalizationDataGatherer& PropertyLocalizationDataGatherer, const EPropertyLocalizationGathererTextFlags GatherTextFlags);
+#endif // WITH_EDITORONLY_DATA
+
+	void RegisterForLocalization();
+
+} // UE::StructUtils::Private
+
+
+FInstancedStructArray::FInstancedStructArray()
+{
+	UE::StructUtils::Private::RegisterForLocalization();
+}
+
+FInstancedStructArray::FInstancedStructArray(const FInstancedStructArray& InOther)
+{
+	Append(InOther);
+}
+
+FInstancedStructArray::FInstancedStructArray(FInstancedStructArray&& InOther)
+	: Memory(InOther.Memory)
+	, AllocatedSize(InOther.AllocatedSize)
+	, NumItems(InOther.NumItems)
+{
+	InOther.Memory = nullptr;
+	InOther.AllocatedSize = 0;
+	InOther.NumItems = 0;
+}
+
+FInstancedStructArray& FInstancedStructArray::operator=(const FInstancedStructArray& InOther)
+{
+	if (this != &InOther)
+	{
+		Reset();
+		Append(InOther);
+	}
+	return *this;
+}
+
+FInstancedStructArray& FInstancedStructArray::operator=(FInstancedStructArray&& InOther)
+{
+	if (this != &InOther)
+	{
+		Empty();
+			
+		Memory = InOther.Memory;
+		AllocatedSize = InOther.AllocatedSize;
+		NumItems = InOther.AllocatedSize;
+			
+		InOther.Memory = nullptr;
+		InOther.AllocatedSize = 0;
+		InOther.NumItems = 0;
+	}
+	return *this;
+}
+
+FInstancedStructArray& FInstancedStructArray::operator=(TConstArrayView<FInstancedStruct> InItems)
+{
+	Reset();
+	TArray<FConstStructView> Views;
+	Views.Reserve(InItems.Num());
+	for (const FInstancedStruct& Value : InItems)
+	{
+		Views.Add(Value);
+	}
+	InsertAt(0, Views);
+	return *this;
+}
+
+FInstancedStructArray& FInstancedStructArray::operator=(TConstArrayView<FConstStructView> InItems)
+{
+	Reset();
+	InsertAt(0, InItems);
+	return *this;
+}
+
+
+void FInstancedStructArray::Append(const FInstancedStructArray& Other)
+{
+	TArray<FConstStructView> Views;
+	for (int32 Index = 0; Index < Other.Num(); Index++)
+	{
+		Views.Add(Other[Index]);
+	}
+	InsertAt(NumItems, Views);
+}
+
+void FInstancedStructArray::Append(TConstArrayView<FInstancedStruct> NewItemValues)
+{
+	TArray<FConstStructView> Views;
+	Views.Reserve(NewItemValues.Num());
+	for (const FInstancedStruct& Value : NewItemValues)
+	{
+		Views.Add(Value);
+	}
+	InsertAt(NumItems, Views);
+}
+
+void FInstancedStructArray::Append(TConstArrayView<FConstStructView> NewItemValues)
+{
+	InsertAt(NumItems, NewItemValues);
+}
+
+void FInstancedStructArray::InsertAt(const int32 InsertAtIndex, const FInstancedStructArray& Other)
+{
+	TArray<FConstStructView> Views;
+	for (int32 Index = 0; Index < Other.Num(); Index++)
+	{
+		Views.Add(Other[Index]);
+	}
+	InsertAt(InsertAtIndex, Views);
+}
+
+void FInstancedStructArray::InsertAt(const int32 InsertAtIndex, TConstArrayView<FInstancedStruct> NewItemValues)
+{
+	TArray<FConstStructView> Views;
+	Views.Reserve(NewItemValues.Num());
+	for (const FInstancedStruct& Value : NewItemValues)
+	{
+		Views.Add(Value);
+	}
+	InsertAt(InsertAtIndex, Views);
+}
+
+void FInstancedStructArray::InsertAt(const int32 InsertAtIndex, TConstArrayView<FConstStructView> ValuesToInsert)
+{
+	checkSlow((ValuesToInsert.Num() >= 0) & (InsertAtIndex >= 0) & (InsertAtIndex <= NumItems));
+
+	const int32 NumAddedItems = ValuesToInsert.Num();
+	const int32 NumHeadItems = InsertAtIndex;
+	const int32 NumTailItems = NumItems - InsertAtIndex;
+
+	TArray<FItem> NewItems;
+	NewItems.Reserve(ValuesToInsert.Num());
+
+	// This alignment is used to align the tail values if needed.
+	int32 OldMinAlignment = DefaultMinAlignment;
+	for (int32 Index = 0; Index < NumItems; Index++)
+	{
+		OldMinAlignment = FMath::Max(OldMinAlignment, GetItem(Index).GetMinAlignment());
+	}
+
+	int32 MinAlignment = OldMinAlignment;
+	int32 Offset = 0;
+
+	// If there are items before the insert location, start the new layout from the end of the last item.
+	if (NumHeadItems > 0)
+	{
+		Offset = GetItem(NumHeadItems - 1).GetEndOffset();
+	}
+	
+	// Layout new items.
+	for (const FConstStructView Value : ValuesToInsert)
+	{
+		FItem& NewItem = NewItems.Emplace_GetRef(Value.GetScriptStruct(), 0);
+		MinAlignment = FMath::Max(MinAlignment, NewItem.GetMinAlignment());
+		Offset = Align(Offset, NewItem.GetMinAlignment());
+		NewItem.Offset = Offset;
+		Offset += NewItem.GetStructureSize();
+	}
+
+	// Move the layout of the tail items. Align them so that we can use one copy to relocate items.
+	int32 NewTailOffset = 0;
+	int32 OldTailOffset = 0;
+	int32 TailSize = 0;
+	if (NumTailItems > 0)
+	{
+		// Align the rest of the items based on the old min alignment, no full layout, just shift the items. 
+		Offset = Align(Offset, OldMinAlignment);
+
+		// Calculate the moved size.
+		NewTailOffset = Offset;
+		OldTailOffset = GetItem(InsertAtIndex).Offset;
+		TailSize = GetItem(NumItems-1).GetEndOffset() - OldTailOffset;
+
+		// Offset items
+		const int32 DeltaOffset = Offset - GetItem(InsertAtIndex).Offset;
+		for (int32 Index = InsertAtIndex; Index < NumItems; Index++)
+		{
+			FItem& OldItem = GetItem(Index);
+			OldItem.Offset += DeltaOffset;
+		}
+
+		// Update offset based on last item.
+		Offset = GetItem(NumItems - 1).GetEndOffset();
+	}
+
+	const int32 NewNumItems = NumItems + NewItems.Num(); 
+	
+	// Space required for the item index.
+	constexpr int32 ItemAlignment = alignof(FItem);
+	MinAlignment = FMath::Max(MinAlignment, ItemAlignment);
+	Offset += NewNumItems * ItemStride;
+	Offset = Align(Offset, ItemAlignment); // Align last, since the items go to the end.
+
+	// Allocate more memory if needed
+	const int32 NewAllocatedSize = Offset;
+	ReserveBytes(NewAllocatedSize, MinAlignment);
+
+	// Set the item index.
+	NumItems = NewNumItems;
+
+	// Move tail items if needed.
+	if (NumTailItems > 0)
+	{
+		FMemory::Memmove(Memory + NewTailOffset, Memory + OldTailOffset, TailSize);
+
+		for (int32 Index = NumItems - 1; Index >= NumItems - NumTailItems; Index--)
+		{
+			GetItem(Index) = GetItem(Index - NumAddedItems);
+		}
+	}
+
+	// Set new items.
+	for (int32 Index = 0; Index < NewItems.Num(); Index++)
+	{
+		GetItem(InsertAtIndex + Index) = NewItems[Index];
+	}
+	
+	// Initialize new item values.
+	for (int32 Index = 0; Index < ValuesToInsert.Num(); Index++)
+	{
+		FConstStructView NewItemValue = ValuesToInsert[Index];
+		const FItem& Item = GetItem(InsertAtIndex + Index);
+		if (Item.ScriptStruct != nullptr)
+		{ 
+			check(Item.ScriptStruct == NewItemValue.GetScriptStruct()); 
+			Item.ScriptStruct->InitializeStruct(Memory + Item.Offset);
+			if (NewItemValue.GetMemory() != nullptr)
+			{
+				Item.ScriptStruct->CopyScriptStruct(Memory + Item.Offset, NewItemValue.GetMemory());
+			}
+		}
+	}
+}
+
+void FInstancedStructArray::RemoveAt(const int32 RemoveAtIndex, const int32 Count)
+{
+	checkSlow((Count >= 0) & (RemoveAtIndex >= 0) & (RemoveAtIndex + Count <= NumItems));
+
+	// Destruct the removed items.
+	for (int32 Index = 0; Index < Count; Index++)
+	{
+		const FItem& Item = GetItem(RemoveAtIndex + Index);
+		if (Item.ScriptStruct != nullptr)
+		{
+			Item.ScriptStruct->DestroyStruct(Memory + Item.Offset);
+		}
+	}
+
+	// Move tail items if any.
+	const int32 FirstTailItem = RemoveAtIndex + Count;
+	const int32 ItemsToMove = NumItems - FirstTailItem;
+	if (ItemsToMove > 0)
+	{
+		// Calculate alignment of the tail items.
+		int32 MinAlignment = DefaultMinAlignment;
+		for (int32 Index = FirstTailItem; Index < NumItems; Index++)
+		{
+			MinAlignment = FMath::Max(MinAlignment, GetItem(Index).GetMinAlignment());
+		}
+
+		// Calculate the tail items new start location.
+		int32 Offset = RemoveAtIndex > 0 ? GetItem(RemoveAtIndex - 1).GetEndOffset() : 0;
+		Offset = Align(Offset, MinAlignment);
+
+		// Move items values.
+		const int32 NewTailOffset = Offset;
+		const int32 OldTailOffset = GetItem(FirstTailItem).Offset;
+		const int32 TailSize = GetItem(NumItems-1).GetEndOffset() - OldTailOffset;
+		if (TailSize > 0)
+		{
+			FMemory::Memmove(Memory + NewTailOffset, Memory + OldTailOffset, TailSize);
+		}
+
+		// Move and update items.
+		const int32 DeltaOffset = NewTailOffset - OldTailOffset;
+		for (int32 Index = FirstTailItem; Index < NumItems; Index++)
+		{
+			FItem& NewItem = GetItem(Index - Count);
+			const FItem& OldItem = GetItem(Index);
+			NewItem.ScriptStruct = OldItem.ScriptStruct;
+			NewItem.Offset = OldItem.Offset + DeltaOffset;
+		}
+	}
+
+	NumItems -= Count;
+}
+
+void FInstancedStructArray::ReserveBytes(const int32 NumBytes, const int32 MinAlignment)
+{
+	if (NumBytes > AllocatedSize || !IsAligned(Memory, MinAlignment))
+	{
+		const int32 NewSize = FMath::Max(AllocatedSize, NumBytes); 
+		uint8* NewMemory = (uint8*)FMemory::Malloc(NewSize, MinAlignment);
+		if (NumItems > 0)
+		{
+			check(Memory != nullptr);
+			// Copy item values (beginning of the buffer).
+			const int32 ValuesSize = GetItem(NumItems - 1).GetEndOffset();
+			check(ValuesSize <= AllocatedSize);
+			FMemory::Memcpy(NewMemory, Memory, ValuesSize);
+			// Copy items (at the end of the buffer).
+			const int32 ItemsSize = NumItems * ItemStride;
+			check(ItemsSize <= AllocatedSize);
+			FMemory::Memcpy(NewMemory + NumBytes - ItemsSize, Memory + AllocatedSize - ItemsSize, ItemsSize);
+		}
+		
+		FMemory::Free(Memory);
+	
+		Memory = NewMemory;
+		AllocatedSize = NewSize;
+	}
+}
+
+void FInstancedStructArray::SetNum(const int32 NewNum)
+{
+	checkSlow(NewNum >= 0 && NewNum <= NumItems);
+
+	for (int32 Index = NewNum; Index < NumItems; Index++)
+	{
+		const FItem& Item = GetItem(Index);
+		if (Item.ScriptStruct != nullptr)
+		{
+			Item.ScriptStruct->DestroyStruct(Memory + Item.Offset);
+		}
+	}
+	NumItems = NewNum;
+}
+
+void FInstancedStructArray::Reset()
+{
+	// Destruct items
+	for (int32 Index = 0; Index < NumItems; Index++)
+	{
+		const FItem& Item = GetItem(Index);
+		if (Item.ScriptStruct != nullptr)
+		{
+			Item.ScriptStruct->DestroyStruct(Memory + Item.Offset);
+		}
+	}
+	NumItems = 0;
+}
+
+void FInstancedStructArray::Empty()
+{
+	Reset();
+
+	// Free memory
+	FMemory::Free(Memory);
+	Memory = nullptr;
+	AllocatedSize = 0;
+}
+
+void FInstancedStructArray::AddStructReferencedObjects(class FReferenceCollector& Collector) const
+{
+	for (int32 Index = 0; Index < NumItems; Index++)
+	{
+		FItem& Item = GetItem(Index);
+		if (Item.ScriptStruct != nullptr)
+		{
+			Collector.AddReferencedObject(Item.ScriptStruct);
+			Collector.AddReferencedObjects(Item.ScriptStruct, Memory + Item.Offset);
+		}
+	}
+}
+
+bool FInstancedStructArray::Identical(const FInstancedStructArray* Other, uint32 PortFlags) const
+{
+	if (Other == nullptr)
+	{
+		return false;
+	}
+
+	// Identical if both are uninitialized.
+	if (!IsValid() && !Other->IsValid())
+	{
+		return true;
+	}
+
+	// Not identical if one is valid and other is not.
+	if (IsValid() != Other->IsValid())
+	{
+		return false;
+	}
+
+	// Not identical if different layouts.
+	if (NumItems != Other->NumItems)
+	{
+		return false;
+	}
+
+	bool bResult = true;
+
+	// Check that the struct contents are identical.
+	for (int32 Index = 0; Index < Num(); Index++)
+	{ 
+		const FItem& Item = GetItem(Index);
+		const FItem& OtherItem = Other->GetItem(Index);
+		if (Item.ScriptStruct != OtherItem.ScriptStruct)
+		{
+			bResult = false;
+			break;
+		}
+		if (Item.ScriptStruct != nullptr && OtherItem.ScriptStruct != nullptr)
+		{
+			const uint8* ItemMemory = Memory + Item.Offset;
+			const uint8* OtherItemMemory = Other->Memory + OtherItem.Offset;
+			
+			if (Item.ScriptStruct->CompareScriptStruct(ItemMemory, OtherItemMemory, PortFlags) == false)
+			{
+				bResult = false;
+				break;
+			}
+		}
+	}
+	
+	return bResult;
+}
+
+bool FInstancedStructArray::Serialize(FArchive& Ar)
+{
+	enum class EVersion : uint8
+	{
+		InitialVersion = 0,
+		// -----<new versions can be added above this line>-----
+		VersionPlusOne,
+		LatestVersion = VersionPlusOne - 1
+	};
+
+	EVersion Version = EVersion::LatestVersion;
+	Ar << Version;
+
+	if (Version > EVersion::LatestVersion)
+	{
+		UE_LOG(LogCore, Error, TEXT("Invalid Version: %hhu"), Version);
+		Ar.SetError();
+		return false;
+	}
+
+	int32 NumItemsSerialized = NumItems;
+	Ar << NumItemsSerialized;
+
+	if (NumItemsSerialized > 0)
+	{
+		if (Ar.IsLoading())
+		{
+			// Load item types 
+			TArray<FConstStructView> Views;
+			Views.Reserve(NumItemsSerialized);
+			for (int32 Index = 0; Index < NumItemsSerialized; Index++)
+			{
+				UScriptStruct* NonConstStruct = nullptr;
+				Ar << NonConstStruct;
+				if (NonConstStruct)
+				{
+					Ar.Preload(NonConstStruct);
+				}
+				Views.Add(FConstStructView(NonConstStruct, nullptr));
+			}
+
+			// Allocate and init items.
+			Empty();
+			InsertAt(0, Views);
+
+			check(NumItems == NumItemsSerialized);
+			
+			// Load item values
+			for (int32 Index = 0; Index < NumItemsSerialized; Index++)
+			{
+				const FItem& Item = GetItem(Index);
+				UScriptStruct* NonConstStruct = const_cast<UScriptStruct*>(Item.ScriptStruct);
+
+				check(NonConstStruct == Item.ScriptStruct);
+				
+				// Size of the serialized memory
+				int32 SerialSize = 0; 
+				Ar << SerialSize;
+
+				// Serialized memory
+				if (NonConstStruct == nullptr && SerialSize > 0)
+				{
+					// A null struct indicates an old struct or an unsupported one for the current target.
+					// In this case we manually seek in the archive to skip its serialized content. 
+					// We don't want to rely on TaggedSerialization that will mark an error in the archive that
+					// may cause other serialization to fail (e.g. FArchive& operator<<(FArchive& Ar, TArray& A))
+					UE_LOG(LogCore, Warning, TEXT("Unable to find serialized UScriptStruct -> Advance %u bytes in the archive and reset to empty FInstancedStructArray"), SerialSize);
+					Ar.Seek(Ar.Tell() + SerialSize);
+				}
+				else if (NonConstStruct != nullptr)
+				{
+					NonConstStruct->SerializeItem(Ar, Memory + Item.Offset, /* Defaults */ nullptr);
+				}
+			}
+		}
+		else if (Ar.IsSaving())
+		{
+			// Save item types
+			for (int32 Index = 0; Index < NumItems; Index++)
+			{
+				const FItem& Item = GetItem(Index);
+				UScriptStruct* NonConstStruct = const_cast<UScriptStruct*>(Item.ScriptStruct);
+				Ar << NonConstStruct;
+			}
+
+			// Save item values
+			for (int32 Index = 0; Index < NumItemsSerialized; Index++)
+			{
+				const FItem& Item = GetItem(Index);
+				UScriptStruct* NonConstStruct = const_cast<UScriptStruct*>(Item.ScriptStruct);
+
+				// Size of the serialized memory (reserve location)
+				const int64 SizeOffset = Ar.Tell(); // Position to write the actual size after struct serialization
+				int32 SerialSize = 0;
+				Ar << SerialSize;
+				
+				// Serialized memory
+				const int64 InitialOffset = Ar.Tell(); // Position before struct serialization to compute its serial size
+				if (NonConstStruct != nullptr)
+				{
+					NonConstStruct->SerializeItem(Ar, Memory + Item.Offset, /* Defaults */ nullptr);
+				}
+				const int64 FinalOffset = Ar.Tell(); // Keep current offset to reset the archive pos after write the serial size
+
+				// Size of the serialized memory
+				Ar.Seek(SizeOffset);	// Go back in the archive to write the actual size
+				SerialSize = (int32)(FinalOffset - InitialOffset);
+				Ar << SerialSize;
+				Ar.Seek(FinalOffset);	// Reset archive to its position
+			}
+		}
+	}
+
+	return true;
+}
+
+void FInstancedStructArray::GetPreloadDependencies(TArray<UObject*>& OutDeps) const
+{
+	for (int32 Index = 0; Index < NumItems; Index++)
+	{
+		FItem& Item = GetItem(Index);
+		if (UScriptStruct* NonConstStruct = const_cast<UScriptStruct*>(Item.ScriptStruct))
+		{
+			OutDeps.Add(NonConstStruct);
+		}
+	}
+}
diff --git a/Engine/Plugins/Experimental/StructUtils/Source/StructUtils/Public/InstancedStructArray.h b/Engine/Plugins/Experimental/StructUtils/Source/StructUtils/Public/InstancedStructArray.h
new file mode 100644
--- /dev/null
+++ b/Engine/Plugins/Experimental/StructUtils/Source/StructUtils/Public/InstancedStructArray.h
@@ -0,0 +1,162 @@
+﻿// Copyright Epic Games, Inc. All Rights Reserved.
+
+#pragma once
+
+#include "InstancedStruct.h"
+#include "StructView.h"
+#include "InstancedStructArray.generated.h"
+
+/**
+ * Array of heterogeneous structs. Can be used as a property, supports serialization,
+ * but does not have type customizations (no editing in the UI).
+ *
+ * If you need UI editable array of heterogeneous structs, use TArray<FInstancedStruct> instead.
+ *
+ * The array item values and the index to an item are stored in one contiguous block of memory.
+ * The size required to specific layout of structs is larger than the sum of their sizes due to alignment,
+ * and because the index to the structs is stored along with the value memory. Each item takes extra 16 bytes
+ * for index. If your items are roughly same size, a TArray<TVariant<>> might be more performant.
+ *
+ * Adding new items is more expensive than on regular TArray<>, layout of the structs needs to be updated,
+ * and initialization is done via UScriptStruct. Adding and removing items should be done in chunks if possible.
+ *
+ * The allocation of new items does not allocate extra space as most array implementations do.
+ * Use Reserve() to reserve certain sized buffer in bytes if that is applicable to your use case.
+ */
+
+USTRUCT()
+struct STRUCTUTILS_API FInstancedStructArray
+{
+	GENERATED_BODY()
+
+	FInstancedStructArray();
+	FInstancedStructArray(const FInstancedStructArray& InOther);
+	FInstancedStructArray(FInstancedStructArray&& InOther);
+	
+	~FInstancedStructArray() { Empty(); }
+
+	FInstancedStructArray& operator=(const FInstancedStructArray& InOther);
+	FInstancedStructArray& operator=(FInstancedStructArray&& InOther);
+	FInstancedStructArray& operator=(TConstArrayView<FInstancedStruct> InItems);
+	FInstancedStructArray& operator=(TConstArrayView<FConstStructView> InItems);
+
+	/** Appends items to the array. */
+	void Append(const FInstancedStructArray& Other);
+	void Append(TConstArrayView<FInstancedStruct> NewItemValues);
+	void Append(TConstArrayView<FConstStructView> NewItemValues);
+
+	/** Insert new items at specified location. */
+	void InsertAt(const int32 InsertAtIndex, const FInstancedStructArray& Other);
+	void InsertAt(const int32 InsertAtIndex, TConstArrayView<FInstancedStruct> ValuesToInsert);
+	void InsertAt(const int32 InsertAtIndex, TConstArrayView<FConstStructView> ValuesToInsert);
+
+	/** Remove items at specific location. Does not change memory allocation. */
+	void RemoveAt(const int32 RemoveAtIndex, const int32 Count);
+
+	/** Reserves at least 'NumBytes' for internal storage. */
+	void ReserveBytes(const int32 NumBytes, const int32 MinAlignment = DefaultMinAlignment);
+	
+	/** Sets the number of items in the array. Note: this can only shrink the array. Does not change memory allocation. */
+	void SetNum(const int32 NewNum);
+
+	/** Returns number of bytes allocated for the array  */
+	int32 GetAllocatedMemory() const { return AllocatedSize; }
+	
+	/** Empties the array, destroys entities. Does not change memory allocation. */
+	void Reset();
+
+	/** Empties the array, destroys entities. Frees memory. */
+	void Empty();
+
+	/** @return true if the instance is correctly initialized. */
+	bool IsValid() const { return NumItems > 0; }
+
+	/** @return Number of items in the instance data. */
+	int32 Num() const { return NumItems; }
+
+	/** @return true of the index is in valid range to the array. */
+	bool IsValidIndex(const int32 Index) const { return Index >= 0 && Index < NumItems; }
+	
+	/** @return view to the struct at specified index. */
+	FStructView operator[](const int32 Index) const
+	{
+		check(IsValid() && IsValidIndex(Index));
+		const FItem& Item = GetItem(Index);
+		return FStructView(Item.ScriptStruct, Memory + Item.Offset);
+	}
+
+	/** Type traits */
+	void AddStructReferencedObjects(class FReferenceCollector& Collector) const;
+	bool Identical(const FInstancedStructArray* Other, uint32 PortFlags) const;
+	bool Serialize(FArchive& Ar);
+	void GetPreloadDependencies(TArray<UObject*>& OutDeps) const;
+
+private:
+
+	/** Struct describing an item in the array. */
+	struct FItem
+	{
+		FItem() = default;
+		FItem(const UScriptStruct* InScriptStruct, const int32 InOffset) : ScriptStruct(InScriptStruct), Offset(InOffset) {}
+
+		int32 GetStructureSize() const
+		{
+			return ScriptStruct != nullptr ? ScriptStruct->GetStructureSize() : 0;
+		}
+
+		int32 GetMinAlignment() const
+		{
+			return ScriptStruct != nullptr ? ScriptStruct->GetMinAlignment() : 1;
+		}
+
+		int32 GetEndOffset() const
+		{
+			return Offset + GetStructureSize();
+		}
+
+		const UScriptStruct* ScriptStruct = nullptr;
+		int32 Offset = 0;
+	};
+
+	/** Default minimum alignment for any allocation. Pointer sized to prevent trivial allocations due to alignment. */
+	static constexpr int32 DefaultMinAlignment = alignof(void*);
+
+	/** Each FItem struct is allocated at the end of the buffer (which is aligned to FItem alignment), and they are stride apart. */
+	static constexpr int32 ItemStride = Align(sizeof(FItem), alignof(FItem));
+
+	/** @returns offset of the item in the allocated memory.  */
+	int32 GetItemOffset(const int32 Index) const
+	{
+		check(NumItems > 0 && Index < NumItems);
+		return AllocatedSize - ((Index + 1) * ItemStride);
+	}
+
+	/** @returns item based on index.  */
+	FItem& GetItem(const int32 Index) const
+	{
+		check(NumItems > 0 && Index < NumItems);
+		const int32 Offset = AllocatedSize - ((Index + 1) * ItemStride);
+		return *(FItem*)(Memory + Offset);
+	}
+
+	/** Memory holding all structs values and items. Values are stored front, and items back. */
+	uint8* Memory = nullptr;
+
+	/** Number of bytes allocated for the array. */
+	int32 AllocatedSize = 0;
+
+	/** Number of items in the array. */
+	int32 NumItems = 0;
+};
+
+template<>
+struct TStructOpsTypeTraits<FInstancedStructArray> : public TStructOpsTypeTraitsBase2<FInstancedStructArray>
+{
+	enum
+	{
+		WithSerializer = true,
+		WithIdentical = true,
+		WithAddStructReferencedObjects = true,
+		WithGetPreloadDependencies = true,
+	};
+};
diff --git a/Engine/Plugins/Experimental/StructUtils/Source/StructUtilsTestSuite/Private/InstancedStructArrayTest.cpp b/Engine/Plugins/Experimental/StructUtils/Source/StructUtilsTestSuite/Private/InstancedStructArrayTest.cpp
new file mode 100644
--- /dev/null
+++ b/Engine/Plugins/Experimental/StructUtils/Source/StructUtilsTestSuite/Private/InstancedStructArrayTest.cpp
@@ -0,0 +1,126 @@
+﻿// Copyright Epic Games, Inc. All Rights Reserved.
+
+#include "AITestsCommon.h"
+#include "InstancedStructArray.h"
+#include "StructUtilsTestTypes.h"
+#include "Serialization/ObjectAndNameAsStringProxyArchive.h"
+
+#define LOCTEXT_NAMESPACE "StructUtilsTests"
+
+PRAGMA_DISABLE_OPTIMIZATION
+
+namespace FInstancedStructArrayTest
+{
+
+struct FTest_BasicInstancedStructArray : FAITestBase
+{
+	virtual bool InstantTest() override
+	{
+		FInstancedStructArray Array;
+
+		TArray<FConstStructView> TestStructs;
+		TestStructs.Add(TBaseStructure<FTestStructSimple>::Get());
+		TestStructs.Add(TBaseStructure<FVector>::Get());
+		TestStructs.Add(TBaseStructure<FTestStructComplex>::Get());
+		
+		Array.Append(TestStructs);
+
+		Array[1].GetMutable<FVector>().X = 42.0;
+
+		AITEST_TRUE(TEXT("Should have 3 items"), Array.Num() == 3);
+		AITEST_TRUE(TEXT("Item 0 should be FTestStructSimple"), Array[0].GetScriptStruct() == TBaseStructure<FTestStructSimple>::Get());
+		AITEST_TRUE(TEXT("Item 1 should be FVector"), Array[1].GetScriptStruct() == TBaseStructure<FVector>::Get());
+		AITEST_TRUE(TEXT("Item 2 should be FTestStructComplex"), Array[2].GetScriptStruct() == TBaseStructure<FTestStructComplex>::Get());
+		AITEST_TRUE(TEXT("Item 1 should have X == 42.0"), FMath::IsNearlyEqual(Array[1].GetMutable<FVector>().X, 42.0));
+
+		Array[2].GetMutable<FTestStructComplex>().StringArray.Add(TEXT("Foo"));
+		
+		TArray<FInstancedStruct> TestInstanced;
+		TestInstanced.AddDefaulted_GetRef().InitializeAs<FTransform>();
+		TestInstanced.AddDefaulted_GetRef().InitializeAs<FVector>(1,2,3);
+		Array.Append(TestInstanced);
+
+		AITEST_TRUE(TEXT("Should have 5 items"), Array.Num() == 5);
+		AITEST_TRUE(TEXT("Item 4 should have Z == 3.0"), FMath::IsNearlyEqual(Array[4].GetMutable<FVector>().Z, 3.0));
+		AITEST_TRUE(TEXT("Item 2 should have text Foo"), Array[2].GetMutable<FTestStructComplex>().StringArray[0] == TEXT("Foo"));
+
+		Array.RemoveAt(2, 1);
+		AITEST_TRUE(TEXT("Should have 4 items"), Array.Num() == 4);
+		AITEST_TRUE(TEXT("Item 2 should be FTransform"), Array[2].GetScriptStruct() == TBaseStructure<FTransform>::Get());
+		AITEST_TRUE(TEXT("Item 3 should be FVector"), Array[3].GetScriptStruct() == TBaseStructure<FVector>::Get());
+		AITEST_TRUE(TEXT("Item 3 should have Z == 3.0"), FMath::IsNearlyEqual(Array[3].GetMutable<FVector>().Z, 3.0));
+
+		Array.SetNum(2);
+		AITEST_TRUE(TEXT("Should have 2 items"), Array.Num() == 2);
+
+		FInstancedStructArray Array2;
+		Array2.Append({ TBaseStructure<FQuat>::Get(), TBaseStructure<FTestStructSimple1>::Get() });
+		Array.InsertAt(1, Array2);
+		AITEST_TRUE(TEXT("Should have 4 items"), Array.Num() == 4);
+		AITEST_TRUE(TEXT("Item 0 should be FTestStructSimple"), Array[0].GetScriptStruct() == TBaseStructure<FTestStructSimple>::Get());
+		AITEST_TRUE(TEXT("Item 1 should be FQuat"), Array[1].GetScriptStruct() == TBaseStructure<FQuat>::Get());
+		AITEST_TRUE(TEXT("Item 2 should be FTestStructSimple1"), Array[2].GetScriptStruct() == TBaseStructure<FTestStructSimple1>::Get());
+		AITEST_TRUE(TEXT("Item 3 should be FVector"), Array[3].GetScriptStruct() == TBaseStructure<FVector>::Get());
+		AITEST_TRUE(TEXT("Item 3 should have X == 42.0"), FMath::IsNearlyEqual(Array[3].GetMutable<FVector>().X, 42.0));
+
+		Array.RemoveAt(0, 3);
+		AITEST_TRUE(TEXT("Should have 1 item"), Array.Num() == 1);
+		AITEST_TRUE(TEXT("Item 0 should be FVector"), Array[0].GetScriptStruct() == TBaseStructure<FVector>::Get());
+		AITEST_TRUE(TEXT("Item 0 should have X == 42.0"), FMath::IsNearlyEqual(Array[0].GetMutable<FVector>().X, 42.0));
+
+		Array.InsertAt(0, Array2);
+		AITEST_TRUE(TEXT("Should have 1 item"), Array.Num() == 3);
+		AITEST_TRUE(TEXT("Item 2 should have X == 42.0"), FMath::IsNearlyEqual(Array[2].GetMutable<FVector>().X, 42.0));
+
+		return true;
+	}
+};
+IMPLEMENT_AI_INSTANT_TEST(FTest_BasicInstancedStructArray, "System.StructUtils.InstancedStructArray.Basic");
+
+struct FTest_SerializeInstancedStructArray : FAITestBase
+{
+	virtual bool InstantTest() override
+	{
+		FInstancedStructArray Array;
+
+		TArray<FInstancedStruct> TestInstanced;
+		TestInstanced.AddDefaulted_GetRef().InitializeAs<FTransform>();
+		TestInstanced.AddDefaulted_GetRef().InitializeAs<FVector>(1,2,3);
+		TestInstanced.AddDefaulted_GetRef().InitializeAs<FTestStructComplex>();
+		Array.Append(TestInstanced);
+
+		Array[2].GetMutable<FTestStructComplex>().StringArray.Add(TEXT("Foo"));
+
+		AITEST_TRUE(TEXT("Should have 3 items"), Array.Num() == 3);
+		AITEST_TRUE(TEXT("Item 2 should have text Foo"), Array[2].GetMutable<FTestStructComplex>().StringArray[0] == TEXT("Foo"));
+
+		TArray<uint8> Memory;
+		
+		FMemoryWriter Writer(Memory);
+		FObjectAndNameAsStringProxyArchive WriterProxy(Writer, /*bInLoadIfFindFails*/false);
+		const bool bSaveResult = Array.Serialize(WriterProxy);
+		AITEST_TRUE(TEXT("Saving should succeed"), bSaveResult);
+
+		FMemoryReader Reader(Memory);
+		FObjectAndNameAsStringProxyArchive ReaderProxy(Reader, /*bInLoadIfFindFails*/true);
+		const bool bLoadResult = Array.Serialize(ReaderProxy);
+		AITEST_TRUE(TEXT("Loading to same array should succeed"), bLoadResult);
+
+		FInstancedStructArray Array2;
+		FMemoryReader Reader2(Memory);
+		FObjectAndNameAsStringProxyArchive ReaderProxy2(Reader2, /*bInLoadIfFindFails*/true);
+		const bool bLoadResult2 = Array2.Serialize(ReaderProxy2);
+		AITEST_TRUE(TEXT("Loading to Array2 should succeed"), bLoadResult2);
+		AITEST_TRUE(TEXT("Array2 should have 3 items"), Array2.Num() == 3);
+		AITEST_TRUE(TEXT("Array2 item 2 should have text Foo"), Array2[2].GetMutable<FTestStructComplex>().StringArray[0] == TEXT("Foo"));
+
+		return true;
+	}
+};
+IMPLEMENT_AI_INSTANT_TEST(FTest_SerializeInstancedStructArray, "System.StructUtils.InstancedStructArray.Serialize");
+
+} // FInstancedStructArrayTest
+
+PRAGMA_ENABLE_OPTIMIZATION
+
+#undef LOCTEXT_NAMESPACE
