InstancedSructArray: Fix for nonunity compile

diff --git a/Engine/Plugins/Experimental/StructUtils/Source/StructUtilsTestSuite/Private/InstancedStructArrayTest.cpp b/Engine/Plugins/Experimental/StructUtils/Source/StructUtilsTestSuite/Private/InstancedStructArrayTest.cpp
--- a/Engine/Plugins/Experimental/StructUtils/Source/StructUtilsTestSuite/Private/InstancedStructArrayTest.cpp
+++ b/Engine/Plugins/Experimental/StructUtils/Source/StructUtilsTestSuite/Private/InstancedStructArrayTest.cpp
@@ -4,6 +4,8 @@
 #include "InstancedStructArray.h"
 #include "StructUtilsTestTypes.h"
 #include "Serialization/ObjectAndNameAsStringProxyArchive.h"
+#include "Serialization/MemoryWriter.h"
+#include "Serialization/MemoryReader.h"
 
 #define LOCTEXT_NAMESPACE "StructUtilsTests"
 
