PropertyBag static analysis fix

diff --git a/Engine/Plugins/Experimental/StructUtils/Source/StructUtils/Private/PropertyBag.cpp b/Engine/Plugins/Experimental/StructUtils/Source/StructUtils/Private/PropertyBag.cpp
--- a/Engine/Plugins/Experimental/StructUtils/Source/StructUtils/Private/PropertyBag.cpp
+++ b/Engine/Plugins/Experimental/StructUtils/Source/StructUtils/Private/PropertyBag.cpp
@@ -1312,11 +1312,8 @@
 
 			const int64 InitialOffset = Ar.Tell(); // Position before struct serialization to compute its serial size
 
-			// BagStruct can be null if it contains structs, classes or enums that are not present anymore. 
-			if (BagStruct != nullptr && Value.GetMutableMemory() != nullptr)
-			{
-				BagStruct->SerializeItem(Ar, Value.GetMutableMemory(), /*Defaults*/nullptr);
-			}
+			check(Value.GetMutableMemory() != nullptr);
+			BagStruct->SerializeItem(Ar, Value.GetMutableMemory(), /*Defaults*/nullptr);
 		
 			const int64 FinalOffset = Ar.Tell(); // Keep current offset to reset the archive pos after write the serial size
 
