Added a function to TStructTypeBitSet to fill an instance with amount of `false` or `true` values coresponding to the number of known types tracked by a given bitset specialization.

diff --git a/Engine/Plugins/Experimental/StructUtils/Source/StructUtils/Public/StructTypeBitSet.h b/Engine/Plugins/Experimental/StructUtils/Source/StructUtils/Public/StructTypeBitSet.h
--- a/Engine/Plugins/Experimental/StructUtils/Source/StructUtils/Public/StructTypeBitSet.h
+++ b/Engine/Plugins/Experimental/StructUtils/Source/StructUtils/Public/StructTypeBitSet.h
@@ -112,6 +112,11 @@ private:
 			return *this;
 		}
 
+		void SetAll(const bool bValue)
+		{
+			Init(bValue, TStructTrackerWrapper::StructTracker.Num());
+		}
+
 		FORCEINLINE bool HasAll(const TBitArray<>& Other) const
 		{
 			FConstWordIterator ThisIterator(*this);
@@ -344,7 +349,7 @@ private:
 
 public:
 
-	static int32 CreateTypeIndex(const TUStructType& InStructType)
+	static int32 GetTypeIndex(const TUStructType& InStructType)
 	{
 #if WITH_STRUCTUTILS_DEBUG
 		ensureMsgf(InStructType.IsChildOf(GetBaseUStruct())
@@ -359,10 +364,10 @@ public:
 	static int32 GetTypeIndex()
 	{
 		static_assert(TIsDerivedFrom<T, TBaseStruct>::IsDerived, "Given struct type doesn't match the expected base struct type.");
-		static const int32 TypeIndex = CreateTypeIndex(*UE::StructUtils::GetAsUStruct<T>());
+		static const int32 TypeIndex = GetTypeIndex(*UE::StructUtils::GetAsUStruct<T>());
 		return TypeIndex;
 	}
-
+	
 	template<typename T>
 	static const TStructTypeBitSet& GetTypeBitSet()
 	{
@@ -400,6 +405,11 @@ public:
 		return StructTypesBitArray.Contains(StructTypeIndex);
 	}
 
+	void SetAll(const bool bValue = true)
+	{
+		StructTypesBitArray.SetAll(bValue);
+	}
+
 	void Add(const TUStructType& InStructType)
 	{
 #if WITH_STRUCTUTILS_DEBUG
