Added `Num` to StructTypeBitSet along with one debug-time funciton fetching the name of value given by index

diff --git a/Engine/Plugins/Experimental/StructUtils/Source/StructUtils/Public/StructTypeBitSet.h b/Engine/Plugins/Experimental/StructUtils/Source/StructUtils/Public/StructTypeBitSet.h
--- a/Engine/Plugins/Experimental/StructUtils/Source/StructUtils/Public/StructTypeBitSet.h
+++ b/Engine/Plugins/Experimental/StructUtils/Source/StructUtils/Public/StructTypeBitSet.h
@@ -536,6 +536,11 @@ public:
 		return FIndexIterator(StructTypesBitArray, bValueToCheck);
 	}
 
+	int32 CountStoredTypes() const
+	{
+		return StructTypesBitArray.CountSetBits();
+	}
+
 	static int32 GetMaxNum() 
 	{		
 		return TStructTrackerWrapper::StructTracker.Num();
@@ -601,6 +606,11 @@ public:
 		}
 	}
 
+	static FName DebugGetStructTypeName(const int32 StructTypeIndex)
+	{
+		return TStructTrackerWrapper::StructTracker.DebugGetStructTypeName(StructTypeIndex);
+	}
+
 	static TConstArrayView<TWeakObjectPtr<const TUStructType>> DebugGetAllStructTypes()
 	{
 		return TStructTrackerWrapper::StructTracker.template DebugGetAllStructTypes<TUStructType>();
