PropertyBag: Ref count UPropertyBags to make them available in C++ constructors

diff --git a/Engine/Plugins/Experimental/StructUtils/Source/StructUtils/Private/PropertyBag.cpp b/Engine/Plugins/Experimental/StructUtils/Source/StructUtils/Private/PropertyBag.cpp
--- a/Engine/Plugins/Experimental/StructUtils/Source/StructUtils/Private/PropertyBag.cpp
+++ b/Engine/Plugins/Experimental/StructUtils/Source/StructUtils/Private/PropertyBag.cpp
@@ -1422,6 +1422,52 @@
 	return NewBag;
 }
 
+void UPropertyBag::InitializeStruct(void* Dest, int32 ArrayDim) const
+{
+	Super::InitializeStruct(Dest, ArrayDim);
+
+	// Do ref counting based on struct usage.
+	// This ensures that if the UPropertyBag is still valid in the C++ destructor of
+	// the last instance of the bag.
+	UPropertyBag* NonConstThis = const_cast<UPropertyBag*>(this);
+	const int32 OldCount = NonConstThis->RefCount.fetch_add(1, std::memory_order_acq_rel);
+	if (OldCount == 0)
+	{
+		NonConstThis->AddToRoot();
+	}
+}
+
+void UPropertyBag::DestroyStruct(void* Dest, int32 ArrayDim) const
+{
+	Super::DestroyStruct(Dest, ArrayDim);
+
+	// Do ref counting based on struct usage.
+	// This ensures that if the UPropertyBag is still valid in the C++ destructor of
+	// the last instance of the bag.
+
+	UPropertyBag* NonConstThis = const_cast<UPropertyBag*>(this);
+	const int32 OldCount = NonConstThis->RefCount.fetch_sub(1, std::memory_order_acq_rel);
+	if (OldCount == 1)
+	{
+		NonConstThis->RemoveFromRoot();
+	}
+	if (OldCount <= 0)
+	{
+		UE_LOG(LogCore, Error, TEXT("PropertyBag: DestroyStruct is called when RefCount is %d."), OldCount);
+	}
+}
+
+void UPropertyBag::FinishDestroy()
+{
+	const int32 Count = RefCount.load(); 
+	if (Count > 0)
+	{
+		UE_LOG(LogCore, Error, TEXT("PropertyBag: Expecting RefCount to be zero on destructor, but it is %d."), Count);
+	}
+	
+	Super::FinishDestroy();
+}
+
 const FPropertyBagPropertyDesc* UPropertyBag::FindPropertyDescByID(const FGuid ID) const
 {
 	return PropertyDescs.FindByPredicate([&ID](const FPropertyBagPropertyDesc& Desc) { return Desc.ID == ID; });
diff --git a/Engine/Plugins/Experimental/StructUtils/Source/StructUtils/Public/PropertyBag.h b/Engine/Plugins/Experimental/StructUtils/Source/StructUtils/Public/PropertyBag.h
--- a/Engine/Plugins/Experimental/StructUtils/Source/StructUtils/Public/PropertyBag.h
+++ b/Engine/Plugins/Experimental/StructUtils/Source/StructUtils/Public/PropertyBag.h
@@ -5,6 +5,7 @@
 #include "InstancedStruct.h"
 #include "StructView.h"
 #include "Templates/ValueOrError.h"
+#include <atomic>
 #include "PropertyBag.generated.h"
 
 /** Property bag property type, loosely based on BluePrint pin types. */
@@ -465,8 +466,15 @@ public:
 	const FPropertyBagPropertyDesc* FindPropertyDescByName(const FName Name) const;
 
 protected:
+
+	virtual void InitializeStruct(void* Dest, int32 ArrayDim = 1) const override;
+	virtual void DestroyStruct(void* Dest, int32 ArrayDim = 1) const override;
+	virtual void FinishDestroy();
+	
 	UPROPERTY()
 	TArray<FPropertyBagPropertyDesc> PropertyDescs;
 
+	std::atomic<int32> RefCount = 0;
+	
 	friend struct FInstancedPropertyBag;
 };
diff --git a/Engine/Plugins/Experimental/StructUtils/Source/StructUtilsTestSuite/Private/PropertyBagTest.cpp b/Engine/Plugins/Experimental/StructUtils/Source/StructUtilsTestSuite/Private/PropertyBagTest.cpp
--- a/Engine/Plugins/Experimental/StructUtils/Source/StructUtilsTestSuite/Private/PropertyBagTest.cpp
+++ b/Engine/Plugins/Experimental/StructUtils/Source/StructUtilsTestSuite/Private/PropertyBagTest.cpp
@@ -256,6 +256,46 @@ struct FTest_Enum : FAITestBase
 };
 IMPLEMENT_AI_INSTANT_TEST(FTest_Enum, "System.StructUtils.PropertyBag.Enum");
 
+struct FTest_GC : FAITestBase
+{
+	virtual bool InstantTest() override
+	{
+		static const FName EnumName(TEXT("Enum"));
+
+		UTestObjectWithPropertyBag* Obj = NewObject<UTestObjectWithPropertyBag>();
+		Obj->Bag.AddProperty(EnumName, EPropertyBagPropertyType::Enum, StaticEnum<EPropertyBagTest1>());
+
+		const UPropertyBag* BagStruct = Obj->Bag.GetPropertyBagStruct();
+		check(BagStruct);
+		
+		const FString BagStructName = BagStruct->GetName();
+		const FString ObjName = Obj->GetName();
+
+		// Obj is unreachable, it should be collected by the GC.
+		Obj = nullptr;
+		CollectGarbage(RF_NoFlags);
+
+		// The used property bag struct should exists after the GC.
+		const UPropertyBag* ExistingObj = FindObject<UPropertyBag>(GetTransientPackage(), *ObjName);
+		const UPropertyBag* ExistingBagStruct1 = FindObject<UPropertyBag>(GetTransientPackage(), *BagStructName);
+
+		AITEST_NULL(TEXT("Obj should have been released"), ExistingObj);
+		AITEST_NOT_NULL(TEXT("Bag struct should exists after Obj released"), ExistingBagStruct1);
+
+		// The next GC should collect the bag struct
+		CollectGarbage(RF_NoFlags);
+
+		const UPropertyBag* ExistingBagStruct2 = FindObject<UPropertyBag>(GetTransientPackage(), *BagStructName);
+		AITEST_NULL(TEXT("Bag struct should not exists after second GC"), ExistingBagStruct2);
+
+		return true;
+	}
+};
+IMPLEMENT_AI_INSTANT_TEST(FTest_GC, "System.StructUtils.PropertyBag.GC");
+
+
+
+
 
 } // FPropertyBagTest
 
diff --git a/Engine/Plugins/Experimental/StructUtils/Source/StructUtilsTestSuite/Public/StructUtilsTestTypes.h b/Engine/Plugins/Experimental/StructUtils/Source/StructUtilsTestSuite/Public/StructUtilsTestTypes.h
--- a/Engine/Plugins/Experimental/StructUtils/Source/StructUtilsTestSuite/Public/StructUtilsTestTypes.h
+++ b/Engine/Plugins/Experimental/StructUtils/Source/StructUtilsTestSuite/Public/StructUtilsTestTypes.h
@@ -3,6 +3,7 @@
 #pragma once
 
 #include "CoreMinimal.h"
+#include "PropertyBag.h"
 #include "StructUtilsTestTypes.generated.h"
 
 USTRUCT()
@@ -146,3 +147,13 @@ enum class EPropertyBagTest2 : uint8
 	Bingo,
 	Bongo,
 };
+
+UCLASS()
+class UTestObjectWithPropertyBag : public UObject
+{
+	GENERATED_BODY()
+
+public:
+	UPROPERTY()
+	FInstancedPropertyBag Bag;
+};
