Property Bag: added support for array properties

diff --git a/Engine/Plugins/Experimental/StructUtils/Source/StructUtils/Private/PropertyBag.cpp b/Engine/Plugins/Experimental/StructUtils/Source/StructUtils/Private/PropertyBag.cpp
--- a/Engine/Plugins/Experimental/StructUtils/Source/StructUtils/Private/PropertyBag.cpp
+++ b/Engine/Plugins/Experimental/StructUtils/Source/StructUtils/Private/PropertyBag.cpp
@@ -3,8 +3,38 @@
 #include "PropertyBag.h"
 #include "StructView.h"
 #include "Hash/CityHash.h"
+#include "Misc/Guid.h"
 #include "UObject/TextProperty.h"
 
+struct STRUCTUTILS_API FPropertyBagCustomVersion
+{
+	enum Type
+	{
+		// Before any version changes were made in the plugin
+		BeforeCustomVersionWasAdded = 0,
+
+		// Added support for array types
+		ContainerTypes = 1,
+
+		// -----<new versions can be added above this line>-------------------------------------------------
+		VersionPlusOne,
+		LatestVersion = VersionPlusOne - 1
+	};
+
+	// The GUID for this custom version number
+	const static FGuid GUID;
+
+private:
+	FPropertyBagCustomVersion() {}
+};
+
+
+const FGuid FPropertyBagCustomVersion::GUID(0x134A157E, 0xD5E249A3, 0x8D4E843C, 0x98FE9E31);
+
+// Register the custom version with core
+FCustomVersionRegistration GPropertyBagCustomVersion(FPropertyBagCustomVersion::GUID, FPropertyBagCustomVersion::LatestVersion, TEXT("PropertyBagCustomVersion"));
+
+
 namespace UE::StructUtils::Private
 {
 
@@ -22,9 +52,9 @@ namespace UE::StructUtils::Private
 	uint64 CalcPropertyDescHash(const FPropertyBagPropertyDesc& Desc)
 	{
 #if WITH_EDITORONLY_DATA
-		const uint32 Hashes[] = { GetTypeHash(Desc.ID), GetTypeHash(Desc.Name), GetTypeHash(Desc.ValueType), GetTypeHash(Desc.MetaData) };
+		const uint32 Hashes[] = { GetTypeHash(Desc.ID), GetTypeHash(Desc.Name), GetTypeHash(Desc.ValueType), GetTypeHash(Desc.ContainerType), GetTypeHash(Desc.MetaData) };
 #else
-		const uint32 Hashes[] = { GetTypeHash(Desc.ID), GetTypeHash(Desc.Name), GetTypeHash(Desc.ValueType) };
+		const uint32 Hashes[] = { GetTypeHash(Desc.ID), GetTypeHash(Desc.Name), GetTypeHash(Desc.ValueType), GetTypeHash(Desc.ContainerType)};
 #endif
 		return CityHash64WithSeed((const char*)Hashes, sizeof(Hashes), GetObjectHash(Desc.ValueTypeObject));
 	}
@@ -39,6 +69,16 @@ namespace UE::StructUtils::Private
 		return Hash;
 	}
 
+	EPropertyBagContainerType GetContainerTypeFromProperty(const FProperty* InSourceProperty)
+	{
+		if (CastField<FArrayProperty>(InSourceProperty))
+		{
+			return EPropertyBagContainerType::Array;
+		}
+		
+		return EPropertyBagContainerType::None;
+	}
+
 	EPropertyBagPropertyType GetValueTypeFromProperty(const FProperty* InSourceProperty)
 	{
 		if (CastField<FBoolProperty>(InSourceProperty))
@@ -102,6 +142,12 @@ namespace UE::StructUtils::Private
 			return EPropertyBagPropertyType::SoftClass;
 		}
 
+		// Handle array property
+		if (const FArrayProperty* ArrayProperty = CastField<FArrayProperty>(InSourceProperty))
+		{
+			return GetValueTypeFromProperty(ArrayProperty->Inner);	
+		}
+
 		return EPropertyBagPropertyType::None;
 	}
 
@@ -138,12 +184,32 @@ namespace UE::StructUtils::Private
 		{
 			return SoftClassProperty->PropertyClass;
 		}
+		
+		// Handle array property
+		if (const FArrayProperty* ArrayProperty = CastField<FArrayProperty>(InSourceProperty))
+		{
+			return GetValueTypeObjectFromProperty(ArrayProperty->Inner);
+		}
+
 
 		return nullptr;
 	}
 
-	FProperty* CreatePropertyFromDesc(const FPropertyBagPropertyDesc& Desc, UScriptStruct* PropertyScope)
+	FProperty* CreatePropertyFromDesc(const FPropertyBagPropertyDesc& Desc, const FFieldVariant PropertyScope)
 	{
+		// Handle array property
+		if (Desc.ContainerType == EPropertyBagContainerType::Array)
+		{
+			FArrayProperty* Prop = new FArrayProperty(PropertyScope, Desc.Name, RF_Public);
+
+			FPropertyBagPropertyDesc InnerDesc = Desc;
+			InnerDesc.Name = FName(InnerDesc.Name.ToString() + TEXT("_inner"));
+			InnerDesc.ContainerType = EPropertyBagContainerType::None;
+			Prop->Inner = CreatePropertyFromDesc(InnerDesc, Prop);
+				
+			return Prop;
+		}
+		
 		switch (Desc.ValueType)
 		{
 		case EPropertyBagPropertyType::Bool:
@@ -280,52 +346,68 @@ namespace UE::StructUtils::Private
 		return nullptr;
 	}
 
-	EPropertyBagResult GetPropertyAsDouble(const FPropertyBagPropertyDesc& Desc, const void* Address, double& OutValue)
+	// Helper functions to get and set property values
+
+	//----------------------------------------------------------------//
+	//  Getters
+	//----------------------------------------------------------------//
+
+	EPropertyBagResult GetPropertyAsInt64(const FPropertyBagPropertyDesc* Desc, const void* Address, int64& OutValue)
 	{
-		check(Desc.CachedProperty);
-		check(Address);
-		
-		switch(Desc.ValueType)
+		if (Desc == nullptr || Desc->CachedProperty == nullptr)
+		{
+			return EPropertyBagResult::PropertyNotFound;
+		}
+		if (Address == nullptr)
+		{
+			return EPropertyBagResult::OutOfBounds;
+		}
+		if (Desc->ContainerType != EPropertyBagContainerType::None)
+		{
+			return EPropertyBagResult::TypeMismatch;
+		}
+
+		switch(Desc->ValueType)
 		{
 		case EPropertyBagPropertyType::Bool:
 			{
-				const FBoolProperty* Property = CastFieldChecked<FBoolProperty>(Desc.CachedProperty);
-				OutValue = Property->GetPropertyValue(Address) ? 1.0 : 0.0;
+				const FBoolProperty* Property = CastFieldChecked<FBoolProperty>(Desc->CachedProperty);
+				OutValue = Property->GetPropertyValue(Address) ? 1 : 0;
 				return EPropertyBagResult::Success;
 			}
 		case EPropertyBagPropertyType::Byte:
 			{
-				const FByteProperty* Property = CastFieldChecked<FByteProperty>(Desc.CachedProperty);
+				const FByteProperty* Property = CastFieldChecked<FByteProperty>(Desc->CachedProperty);
 				OutValue = Property->GetPropertyValue(Address);
 				return EPropertyBagResult::Success;
 			}
 		case EPropertyBagPropertyType::Int32:
 			{
-				const FIntProperty* Property = CastFieldChecked<FIntProperty>(Desc.CachedProperty);
+				const FIntProperty* Property = CastFieldChecked<FIntProperty>(Desc->CachedProperty);
 				OutValue = Property->GetPropertyValue(Address);
 				return EPropertyBagResult::Success;
 			}
 		case EPropertyBagPropertyType::Int64:
 			{
-				const FInt64Property* Property = CastFieldChecked<FInt64Property>(Desc.CachedProperty);
+				const FInt64Property* Property = CastFieldChecked<FInt64Property>(Desc->CachedProperty);
 				OutValue = Property->GetPropertyValue(Address);
 				return EPropertyBagResult::Success;
 			}
 		case EPropertyBagPropertyType::Float:
 			{
-				const FFloatProperty* Property = CastFieldChecked<FFloatProperty>(Desc.CachedProperty);
-				OutValue = Property->GetPropertyValue(Address);
+				const FFloatProperty* Property = CastFieldChecked<FFloatProperty>(Desc->CachedProperty);
+				OutValue = (int64)Property->GetPropertyValue(Address);
 				return EPropertyBagResult::Success;
 			}
 		case EPropertyBagPropertyType::Double:
 			{
-				const FDoubleProperty* Property = CastFieldChecked<FDoubleProperty>(Desc.CachedProperty);
-				OutValue = Property->GetPropertyValue(Address);
+				const FDoubleProperty* Property = CastFieldChecked<FDoubleProperty>(Desc->CachedProperty);
+				OutValue = (int64)Property->GetPropertyValue(Address);
 				return EPropertyBagResult::Success;
 			}
 		case EPropertyBagPropertyType::Enum:
 			{
-				const FEnumProperty* EnumProperty = CastFieldChecked<FEnumProperty>(Desc.CachedProperty);
+				const FEnumProperty* EnumProperty = CastFieldChecked<FEnumProperty>(Desc->CachedProperty);
 				const FNumericProperty* UnderlyingProperty = EnumProperty->GetUnderlyingProperty();
 				check(UnderlyingProperty);
 				OutValue = UnderlyingProperty->GetSignedIntPropertyValue(Address);
@@ -336,55 +418,65 @@ namespace UE::StructUtils::Private
 		}
 	}
 
-	EPropertyBagResult SetPropertyFromDouble(const FPropertyBagPropertyDesc& Desc, void* Address, const double InValue)
+	EPropertyBagResult GetPropertyAsDouble(const FPropertyBagPropertyDesc* Desc, const void* Address, double& OutValue)
 	{
-		check(Desc.CachedProperty);
-		check(Address);
-		
-		switch(Desc.ValueType)
+		if (Desc == nullptr || Desc->CachedProperty == nullptr)
+		{
+			return EPropertyBagResult::PropertyNotFound;
+		}
+		if (Address == nullptr)
+		{
+			return EPropertyBagResult::OutOfBounds;
+		}
+		if (Desc->ContainerType != EPropertyBagContainerType::None)
+		{
+			return EPropertyBagResult::TypeMismatch;
+		}
+
+		switch(Desc->ValueType)
 		{
 		case EPropertyBagPropertyType::Bool:
 			{
-				const FBoolProperty* Property = CastFieldChecked<FBoolProperty>(Desc.CachedProperty);
-				Property->SetPropertyValue(Address, FMath::IsNearlyZero(InValue) ? false : true);
+				const FBoolProperty* Property = CastFieldChecked<FBoolProperty>(Desc->CachedProperty);
+				OutValue = Property->GetPropertyValue(Address) ? 1.0 : 0.0;
 				return EPropertyBagResult::Success;
 			}
 		case EPropertyBagPropertyType::Byte:
 			{
-				const FByteProperty* Property = CastFieldChecked<FByteProperty>(Desc.CachedProperty);
-				Property->SetPropertyValue(Address, FMath::RoundToInt32(InValue));
+				const FByteProperty* Property = CastFieldChecked<FByteProperty>(Desc->CachedProperty);
+				OutValue = Property->GetPropertyValue(Address);
 				return EPropertyBagResult::Success;
 			}
 		case EPropertyBagPropertyType::Int32:
 			{
-				const FIntProperty* Property = CastFieldChecked<FIntProperty>(Desc.CachedProperty);
-				Property->SetPropertyValue(Address, FMath::RoundToInt32(InValue));
+				const FIntProperty* Property = CastFieldChecked<FIntProperty>(Desc->CachedProperty);
+				OutValue = Property->GetPropertyValue(Address);
 				return EPropertyBagResult::Success;
 			}
 		case EPropertyBagPropertyType::Int64:
 			{
-				const FInt64Property* Property = CastFieldChecked<FInt64Property>(Desc.CachedProperty);
-				Property->SetPropertyValue(Address, FMath::RoundToInt64(InValue));
+				const FInt64Property* Property = CastFieldChecked<FInt64Property>(Desc->CachedProperty);
+				OutValue = Property->GetPropertyValue(Address);
 				return EPropertyBagResult::Success;
 			}
 		case EPropertyBagPropertyType::Float:
 			{
-				const FFloatProperty* Property = CastFieldChecked<FFloatProperty>(Desc.CachedProperty);
-				Property->SetPropertyValue(Address, (float)InValue);
+				const FFloatProperty* Property = CastFieldChecked<FFloatProperty>(Desc->CachedProperty);
+				OutValue = Property->GetPropertyValue(Address);
 				return EPropertyBagResult::Success;
 			}
 		case EPropertyBagPropertyType::Double:
 			{
-				const FDoubleProperty* Property = CastFieldChecked<FDoubleProperty>(Desc.CachedProperty);
-				Property->SetPropertyValue(Address, InValue);
+				const FDoubleProperty* Property = CastFieldChecked<FDoubleProperty>(Desc->CachedProperty);
+				OutValue = Property->GetPropertyValue(Address);
 				return EPropertyBagResult::Success;
 			}
 		case EPropertyBagPropertyType::Enum:
 			{
-				const FEnumProperty* EnumProperty = CastFieldChecked<FEnumProperty>(Desc.CachedProperty);
+				const FEnumProperty* EnumProperty = CastFieldChecked<FEnumProperty>(Desc->CachedProperty);
 				const FNumericProperty* UnderlyingProperty = EnumProperty->GetUnderlyingProperty();
 				check(UnderlyingProperty);
-				UnderlyingProperty->SetIntPropertyValue(Address, (uint64)InValue);
+				OutValue = UnderlyingProperty->GetSignedIntPropertyValue(Address);
 				return EPropertyBagResult::Success;
 			}
 		default:
@@ -392,55 +484,197 @@ namespace UE::StructUtils::Private
 		}
 	}
 
-	EPropertyBagResult GetPropertyAsInt64(const FPropertyBagPropertyDesc& Desc, const void* Address, int64& OutValue)
+	// Generic property getter. Used for FName, FString, FText. 
+	template<typename T, typename PropT>
+	EPropertyBagResult GetPropertyValue(const FPropertyBagPropertyDesc* Desc, const void* Address, T& OutValue)
+	{
+		if (Desc == nullptr || Desc->CachedProperty == nullptr)
+    	{
+    		return EPropertyBagResult::PropertyNotFound;
+    	}
+		if (Address == nullptr)
+		{
+			return EPropertyBagResult::OutOfBounds;
+		}
+		if (Desc->ContainerType != EPropertyBagContainerType::None)
+		{
+			return EPropertyBagResult::TypeMismatch;
+		}
+
+		if (!Desc->CachedProperty->IsA<PropT>())
+		{
+			return EPropertyBagResult::TypeMismatch;
+		}
+		
+		const PropT* Property = CastFieldChecked<PropT>(Desc->CachedProperty);
+		OutValue = Property->GetPropertyValue(Address);
+
+		return EPropertyBagResult::Success;
+	}
+
+	EPropertyBagResult GetPropertyValueAsEnum(const FPropertyBagPropertyDesc* Desc, const UEnum* RequestedEnum, const void* Address, uint8& OutValue)
+	{
+		if (Desc == nullptr || Desc->CachedProperty == nullptr)
+    	{
+    		return EPropertyBagResult::PropertyNotFound;
+    	}
+		if (Desc->ValueType != EPropertyBagPropertyType::Enum)
+		{
+			return EPropertyBagResult::TypeMismatch;
+		}
+		if (Address == nullptr)
+		{
+			return EPropertyBagResult::OutOfBounds;
+		}
+		if (Desc->ContainerType != EPropertyBagContainerType::None)
+		{
+			return EPropertyBagResult::TypeMismatch;
+		}
+
+		const FEnumProperty* EnumProperty = CastFieldChecked<FEnumProperty>(Desc->CachedProperty);
+		const FNumericProperty* UnderlyingProperty = EnumProperty->GetUnderlyingProperty();
+		check(UnderlyingProperty);
+	
+		if (RequestedEnum != EnumProperty->GetEnum())
+		{
+			return EPropertyBagResult::TypeMismatch;
+		}
+	
+		OutValue = (uint8)UnderlyingProperty->GetUnsignedIntPropertyValue(Address);
+
+		return EPropertyBagResult::Success;
+	}
+
+	EPropertyBagResult GetPropertyValueAsStruct(const FPropertyBagPropertyDesc* Desc, const UScriptStruct* RequestedStruct, const void* Address, FStructView& OutValue)
+	{
+		if (Desc == nullptr || Desc->CachedProperty == nullptr)
+		{
+			return EPropertyBagResult::PropertyNotFound;
+		}
+		if (Desc->ValueType != EPropertyBagPropertyType::Struct)
+		{
+			return EPropertyBagResult::TypeMismatch;
+		}
+		if (Address == nullptr)
+		{
+			return EPropertyBagResult::OutOfBounds;
+		}
+		if (Desc->ContainerType != EPropertyBagContainerType::None)
+		{
+			return EPropertyBagResult::TypeMismatch;
+		}
+
+		const FStructProperty* StructProperty = CastFieldChecked<FStructProperty>(Desc->CachedProperty);
+		check(StructProperty->Struct);
+
+		if (RequestedStruct != nullptr && CanCastTo(StructProperty->Struct, RequestedStruct) == false)
+		{
+			return EPropertyBagResult::TypeMismatch;
+		}
+	
+		OutValue = FStructView(StructProperty->Struct, (uint8*)Address);
+
+		return EPropertyBagResult::Success;
+	}
+
+	EPropertyBagResult GetPropertyValueAsObject(const FPropertyBagPropertyDesc* Desc, const UClass* RequestedClass, const void* Address, UObject*& OutValue)
+	{
+		if (Desc == nullptr || Desc->CachedProperty == nullptr)
+		{
+			return EPropertyBagResult::PropertyNotFound;
+		}
+		if (Desc->ValueType != EPropertyBagPropertyType::Object
+			&& Desc->ValueType != EPropertyBagPropertyType::SoftObject
+			&& Desc->ValueType != EPropertyBagPropertyType::Class
+			&& Desc->ValueType != EPropertyBagPropertyType::SoftClass)
+		{
+			return EPropertyBagResult::TypeMismatch;
+		}
+		if (Address == nullptr)
+		{
+			return EPropertyBagResult::OutOfBounds;
+		}
+		if (Desc->ContainerType != EPropertyBagContainerType::None)
+		{
+			return EPropertyBagResult::TypeMismatch;
+		}
+
+		const FObjectPropertyBase* ObjectProperty = CastFieldChecked<FObjectPropertyBase>(Desc->CachedProperty);
+		check(ObjectProperty->PropertyClass);
+
+		if (RequestedClass != nullptr && CanCastTo(ObjectProperty->PropertyClass, RequestedClass) == false)
+		{
+			return EPropertyBagResult::TypeMismatch;
+		}
+	
+		OutValue = ObjectProperty->GetObjectPropertyValue(Address);
+
+		return EPropertyBagResult::Success;
+	}
+
+	//----------------------------------------------------------------//
+	//  Setters
+	//----------------------------------------------------------------//
+
+	EPropertyBagResult SetPropertyFromInt64(const FPropertyBagPropertyDesc* Desc, void* Address, const int64 InValue)
 	{
-		check(Desc.CachedProperty);
-		check(Address);
+		if (Desc == nullptr || Desc->CachedProperty == nullptr)
+		{
+			return EPropertyBagResult::PropertyNotFound;
+		}
+		if (Address == nullptr)
+		{
+			return EPropertyBagResult::OutOfBounds;
+		}
+		if (Desc->ContainerType != EPropertyBagContainerType::None)
+		{
+			return EPropertyBagResult::TypeMismatch;
+		}
 		
-		switch(Desc.ValueType)
+		switch(Desc->ValueType)
 		{
 		case EPropertyBagPropertyType::Bool:
 			{
-				const FBoolProperty* Property = CastFieldChecked<FBoolProperty>(Desc.CachedProperty);
-				OutValue = Property->GetPropertyValue(Address) ? 1 : 0;
+				const FBoolProperty* Property = CastFieldChecked<FBoolProperty>(Desc->CachedProperty);
+				Property->SetPropertyValue(Address, InValue != 0);
 				return EPropertyBagResult::Success;
 			}
 		case EPropertyBagPropertyType::Byte:
 			{
-				const FByteProperty* Property = CastFieldChecked<FByteProperty>(Desc.CachedProperty);
-				OutValue = Property->GetPropertyValue(Address);
+				const FByteProperty* Property = CastFieldChecked<FByteProperty>(Desc->CachedProperty);
+				Property->SetPropertyValue(Address, (uint8)InValue);
 				return EPropertyBagResult::Success;
 			}
 		case EPropertyBagPropertyType::Int32:
 			{
-				const FIntProperty* Property = CastFieldChecked<FIntProperty>(Desc.CachedProperty);
-				OutValue = Property->GetPropertyValue(Address);
+				const FIntProperty* Property = CastFieldChecked<FIntProperty>(Desc->CachedProperty);
+				Property->SetPropertyValue(Address, (int32)InValue);
 				return EPropertyBagResult::Success;
 			}
 		case EPropertyBagPropertyType::Int64:
 			{
-				const FInt64Property* Property = CastFieldChecked<FInt64Property>(Desc.CachedProperty);
-				OutValue = Property->GetPropertyValue(Address);
+				const FInt64Property* Property = CastFieldChecked<FInt64Property>(Desc->CachedProperty);
+				Property->SetPropertyValue(Address, InValue);
 				return EPropertyBagResult::Success;
 			}
 		case EPropertyBagPropertyType::Float:
 			{
-				const FFloatProperty* Property = CastFieldChecked<FFloatProperty>(Desc.CachedProperty);
-				OutValue = (int64)Property->GetPropertyValue(Address);
+				const FFloatProperty* Property = CastFieldChecked<FFloatProperty>(Desc->CachedProperty);
+				Property->SetPropertyValue(Address, (float)InValue);
 				return EPropertyBagResult::Success;
 			}
 		case EPropertyBagPropertyType::Double:
 			{
-				const FDoubleProperty* Property = CastFieldChecked<FDoubleProperty>(Desc.CachedProperty);
-				OutValue = (int64)Property->GetPropertyValue(Address);
+				const FDoubleProperty* Property = CastFieldChecked<FDoubleProperty>(Desc->CachedProperty);
+				Property->SetPropertyValue(Address, (double)InValue);
 				return EPropertyBagResult::Success;
 			}
 		case EPropertyBagPropertyType::Enum:
 			{
-				const FEnumProperty* EnumProperty = CastFieldChecked<FEnumProperty>(Desc.CachedProperty);
+				const FEnumProperty* EnumProperty = CastFieldChecked<FEnumProperty>(Desc->CachedProperty);
 				const FNumericProperty* UnderlyingProperty = EnumProperty->GetUnderlyingProperty();
 				check(UnderlyingProperty);
-				OutValue = UnderlyingProperty->GetSignedIntPropertyValue(Address);
+				UnderlyingProperty->SetIntPropertyValue(Address, (uint64)InValue);
 				return EPropertyBagResult::Success;
 			}
 		default:
@@ -448,52 +682,62 @@ namespace UE::StructUtils::Private
 		}
 	}
 
-	EPropertyBagResult SetPropertyFromInt64(const FPropertyBagPropertyDesc& Desc, void* Address, const int64 InValue)
+	EPropertyBagResult SetPropertyFromDouble(const FPropertyBagPropertyDesc* Desc, void* Address, const double InValue)
 	{
-		check(Desc.CachedProperty);
-		check(Address);
-		
-		switch(Desc.ValueType)
+		if (Desc == nullptr || Desc->CachedProperty == nullptr)
+		{
+			return EPropertyBagResult::PropertyNotFound;
+		}
+		if (Address == nullptr)
+		{
+			return EPropertyBagResult::OutOfBounds;
+		}
+		if (Desc->ContainerType != EPropertyBagContainerType::None)
+		{
+			return EPropertyBagResult::TypeMismatch;
+		}
+
+		switch(Desc->ValueType)
 		{
 		case EPropertyBagPropertyType::Bool:
 			{
-				const FBoolProperty* Property = CastFieldChecked<FBoolProperty>(Desc.CachedProperty);
-				Property->SetPropertyValue(Address, InValue != 0);
+				const FBoolProperty* Property = CastFieldChecked<FBoolProperty>(Desc->CachedProperty);
+				Property->SetPropertyValue(Address, FMath::IsNearlyZero(InValue) ? false : true);
 				return EPropertyBagResult::Success;
 			}
 		case EPropertyBagPropertyType::Byte:
 			{
-				const FByteProperty* Property = CastFieldChecked<FByteProperty>(Desc.CachedProperty);
-				Property->SetPropertyValue(Address, (uint8)InValue);
+				const FByteProperty* Property = CastFieldChecked<FByteProperty>(Desc->CachedProperty);
+				Property->SetPropertyValue(Address, FMath::RoundToInt32(InValue));
 				return EPropertyBagResult::Success;
 			}
 		case EPropertyBagPropertyType::Int32:
 			{
-				const FIntProperty* Property = CastFieldChecked<FIntProperty>(Desc.CachedProperty);
-				Property->SetPropertyValue(Address, (int32)InValue);
+				const FIntProperty* Property = CastFieldChecked<FIntProperty>(Desc->CachedProperty);
+				Property->SetPropertyValue(Address, FMath::RoundToInt32(InValue));
 				return EPropertyBagResult::Success;
 			}
 		case EPropertyBagPropertyType::Int64:
 			{
-				const FInt64Property* Property = CastFieldChecked<FInt64Property>(Desc.CachedProperty);
-				Property->SetPropertyValue(Address, InValue);
+				const FInt64Property* Property = CastFieldChecked<FInt64Property>(Desc->CachedProperty);
+				Property->SetPropertyValue(Address, FMath::RoundToInt64(InValue));
 				return EPropertyBagResult::Success;
 			}
 		case EPropertyBagPropertyType::Float:
 			{
-				const FFloatProperty* Property = CastFieldChecked<FFloatProperty>(Desc.CachedProperty);
+				const FFloatProperty* Property = CastFieldChecked<FFloatProperty>(Desc->CachedProperty);
 				Property->SetPropertyValue(Address, (float)InValue);
 				return EPropertyBagResult::Success;
 			}
 		case EPropertyBagPropertyType::Double:
 			{
-				const FDoubleProperty* Property = CastFieldChecked<FDoubleProperty>(Desc.CachedProperty);
-				Property->SetPropertyValue(Address, (double)InValue);
+				const FDoubleProperty* Property = CastFieldChecked<FDoubleProperty>(Desc->CachedProperty);
+				Property->SetPropertyValue(Address, InValue);
 				return EPropertyBagResult::Success;
 			}
 		case EPropertyBagPropertyType::Enum:
 			{
-				const FEnumProperty* EnumProperty = CastFieldChecked<FEnumProperty>(Desc.CachedProperty);
+				const FEnumProperty* EnumProperty = CastFieldChecked<FEnumProperty>(Desc->CachedProperty);
 				const FNumericProperty* UnderlyingProperty = EnumProperty->GetUnderlyingProperty();
 				check(UnderlyingProperty);
 				UnderlyingProperty->SetIntPropertyValue(Address, (uint64)InValue);
@@ -504,176 +748,236 @@ namespace UE::StructUtils::Private
 		}
 	}
 
-	void CopyMatchingValuesByID(const FConstStructView Source, const FStructView Target)
+	// Generic property setter. Used for FName, FString, FText 
+	template<typename T, typename PropT>
+	EPropertyBagResult SetPropertyValue(const FPropertyBagPropertyDesc* Desc, void* Address, const T& InValue)
 	{
-		if (!Source.IsValid() || !Target.IsValid())
+		if (Desc == nullptr || Desc->CachedProperty == nullptr)
 		{
-			return;
+			return EPropertyBagResult::PropertyNotFound;
 		}
-
-		const UPropertyBag* SourceBagStruct = Cast<const UPropertyBag>(Source.GetScriptStruct());
-		const UPropertyBag* TargetBagStruct = Cast<const UPropertyBag>(Target.GetScriptStruct());
-
-		if (!SourceBagStruct || !TargetBagStruct)
+		if (Address == nullptr)
 		{
-			return;
+			return EPropertyBagResult::OutOfBounds;
 		}
-
-		// Iterate over source and copy to target if possible. Source is expected to usually have less items.
-		for (const FPropertyBagPropertyDesc& SourceDesc : SourceBagStruct->GetPropertyDescs())
+		if (Desc->ContainerType != EPropertyBagContainerType::None)
 		{
-			const FPropertyBagPropertyDesc* PotentialTargetDesc = TargetBagStruct->FindPropertyDescByID(SourceDesc.ID);
-			if (PotentialTargetDesc == nullptr
-				|| PotentialTargetDesc->CachedProperty == nullptr
-				|| SourceDesc.CachedProperty == nullptr)
-			{
-				continue;
-			}
-
-			const FPropertyBagPropertyDesc& TargetDesc = *PotentialTargetDesc;
-			void* TargetAddress = Target.GetMutableMemory() + TargetDesc.CachedProperty->GetOffset_ForInternal();
-			const void* SourceAddress = Source.GetMemory() + SourceDesc.CachedProperty->GetOffset_ForInternal();
-			
-			if (TargetDesc.CompatibleType(SourceDesc))
-			{
-				TargetDesc.CachedProperty->CopyCompleteValue(TargetAddress, SourceAddress);
-			}
-			else if (TargetDesc.IsNumericType() && SourceDesc.IsNumericType())
-			{
-				// Try to convert numeric types.
-				if (TargetDesc.IsNumericFloatType())
-				{
-					double Value = 0;
-					if (GetPropertyAsDouble(SourceDesc, SourceAddress, Value) == EPropertyBagResult::Success)
-					{
-						SetPropertyFromDouble(TargetDesc, TargetAddress, Value);
-					}
-				}
-				else
-				{
-					int64 Value = 0;
-					if (GetPropertyAsInt64(SourceDesc, SourceAddress, Value) == EPropertyBagResult::Success)
-					{
-						SetPropertyFromInt64(TargetDesc, TargetAddress, Value);
-					}
-				}
-			}
-			else if ((TargetDesc.IsObjectType() && SourceDesc.IsObjectType())
-				|| (TargetDesc.IsClassType() && SourceDesc.IsClassType()))
-			{
-				// Try convert between compatible objects and classes.
-				const UClass* TargetObjectClass = Cast<const UClass>(TargetDesc.ValueTypeObject);
-				const UClass* SourceObjectClass = Cast<const UClass>(SourceDesc.ValueTypeObject);
-				if (CanCastTo(SourceObjectClass, TargetObjectClass))
-				{
-					const FObjectPropertyBase* TargetProp = CastFieldChecked<FObjectPropertyBase>(TargetDesc.CachedProperty);
-					const FObjectPropertyBase* SourceProp = CastFieldChecked<FObjectPropertyBase>(SourceDesc.CachedProperty);
-					TargetProp->SetObjectPropertyValue(TargetAddress, SourceProp->GetObjectPropertyValue(SourceAddress));
-				}
-			}
+			return EPropertyBagResult::TypeMismatch;
 		}
-	}
 
-	// Helper templates to reduce repeated work when dealing with property access.
-
-	template<typename T>
-	TValueOrError<T, EPropertyBagResult> GetValueInt64(const FInstancedPropertyBag& Bag, const FName Name)
-	{
-		const FPropertyBagPropertyDesc* Desc = Bag.FindPropertyDescByName(Name);
-		if (Desc == nullptr || Desc->CachedProperty == nullptr)
+		if (!Desc->CachedProperty->IsA<PropT>())
 		{
-			return MakeError(EPropertyBagResult::PropertyNotFound);
+			return EPropertyBagResult::TypeMismatch;
 		}
-	 
-		check(Bag.GetValue().IsValid());
-		const void* Address = Bag.GetValue().GetMemory() + Desc->CachedProperty->GetOffset_ForInternal();
-		int64 IntValue = 0;
-		const EPropertyBagResult Result = GetPropertyAsInt64(*Desc, Address, IntValue);
-		if (Result != EPropertyBagResult::Success)
-		{
-			return MakeError(Result);
-		}
-		return MakeValue(IntValue);
+		
+		const PropT* Property = CastFieldChecked<PropT>(Desc->CachedProperty);
+		Property->SetPropertyValue(Address, InValue);
+
+		return EPropertyBagResult::Success;
 	}
 
-	template<typename T>
-	TValueOrError<T, EPropertyBagResult> GetValueDouble(const FInstancedPropertyBag& Bag, const FName Name)
+	EPropertyBagResult SetPropertyValueAsEnum(const FPropertyBagPropertyDesc* Desc, void* Address, const uint8 InValue, const UEnum* Enum)
 	{
-		const FPropertyBagPropertyDesc* Desc = Bag.FindPropertyDescByName(Name);
 		if (Desc == nullptr || Desc->CachedProperty == nullptr)
 		{
-			return MakeError(EPropertyBagResult::PropertyNotFound);
+			return EPropertyBagResult::PropertyNotFound;
 		}
-
-		check(Bag.GetValue().IsValid());
-		const void* Address = Bag.GetValue().GetMemory() + Desc->CachedProperty->GetOffset_ForInternal();
-		double DblValue = 0;
-		const EPropertyBagResult Result = GetPropertyAsDouble(*Desc, Address, DblValue);
-		if (Result != EPropertyBagResult::Success)
+		if (Desc->ValueType != EPropertyBagPropertyType::Enum)
+		{
+			return EPropertyBagResult::TypeMismatch;
+		}
+		if (Address == nullptr)
 		{
-			return MakeError(Result);
+			return EPropertyBagResult::OutOfBounds;
+		}
+		if (Desc->ContainerType != EPropertyBagContainerType::None)
+		{
+			return EPropertyBagResult::TypeMismatch;
 		}
-		return MakeValue(DblValue);
-	}
 
-	template<typename T, typename PropT>
-	TValueOrError<T, EPropertyBagResult> GetValue(const FInstancedPropertyBag& Bag, const FName Name)
-	{
-		const FPropertyBagPropertyDesc* Desc = Bag.FindPropertyDescByName(Name);
-		if (Desc == nullptr || Desc->CachedProperty == nullptr || !Desc->CachedProperty->IsA<PropT>())
+		const FEnumProperty* EnumProperty = CastFieldChecked<FEnumProperty>(Desc->CachedProperty);
+		const FNumericProperty* UnderlyingProperty = EnumProperty->GetUnderlyingProperty();
+		check(UnderlyingProperty);
+
+		if (Enum != EnumProperty->GetEnum())
 		{
-			return MakeError(EPropertyBagResult::PropertyNotFound);
+			return EPropertyBagResult::TypeMismatch;
 		}
-		check(Desc->CachedProperty);
-		const PropT* Property = CastFieldChecked<PropT>(Desc->CachedProperty);
-		check(Bag.GetValue().IsValid());
-		const void* Address = Bag.GetValue().GetMemory() + Desc->CachedProperty->GetOffset_ForInternal();
-		return MakeValue(Property->GetPropertyValue(Address));
+	
+		UnderlyingProperty->SetIntPropertyValue(Address, (uint64)InValue);
+	
+		return EPropertyBagResult::Success;
 	}
 
-
-	template<typename T>
-	EPropertyBagResult SetValueInt64(const FInstancedPropertyBag& Bag, const FName Name, const T Value)
+	EPropertyBagResult SetPropertyValueAsStruct(const FPropertyBagPropertyDesc* Desc, void* Address, const FConstStructView InValue)
 	{
-		const FPropertyBagPropertyDesc* Desc = Bag.FindPropertyDescByName(Name);
 		if (Desc == nullptr || Desc->CachedProperty == nullptr)
 		{
 			return EPropertyBagResult::PropertyNotFound;
 		}
-		check(Desc->CachedProperty);
-		check(Bag.GetValue().IsValid());
-		void* Address = Bag.GetMutableValue().GetMutableMemory() + Desc->CachedProperty->GetOffset_ForInternal();
-		return SetPropertyFromInt64(*Desc, Address, (int64)Value);
+		if (Desc->ValueType != EPropertyBagPropertyType::Struct)
+		{
+			return EPropertyBagResult::TypeMismatch;
+		}
+		if (Address == nullptr)
+		{
+			return EPropertyBagResult::OutOfBounds;
+		}
+		if (Desc->ContainerType != EPropertyBagContainerType::None)
+		{
+			return EPropertyBagResult::TypeMismatch;
+		}
+
+		const FStructProperty* StructProperty = CastFieldChecked<FStructProperty>(Desc->CachedProperty);
+		check(StructProperty->Struct);
+
+		if (InValue.IsValid())
+		{
+			if (InValue.GetScriptStruct() != StructProperty->Struct)
+			{
+				return EPropertyBagResult::TypeMismatch;
+			}
+			StructProperty->Struct->CopyScriptStruct(Address, InValue.GetMemory());
+		}
+		else
+		{
+			StructProperty->Struct->ClearScriptStruct(Address);
+		}
+
+		return EPropertyBagResult::Success;
 	}
 
-	template<typename T>
-	EPropertyBagResult SetValueDouble(const FInstancedPropertyBag& Bag, const FName Name, const T Value)
+	EPropertyBagResult SetPropertyValueAsObject(const FPropertyBagPropertyDesc* Desc, void* Address, UObject* InValue)
 	{
-		const FPropertyBagPropertyDesc* Desc = Bag.FindPropertyDescByName(Name);
 		if (Desc == nullptr || Desc->CachedProperty == nullptr)
 		{
 			return EPropertyBagResult::PropertyNotFound;
 		}
-		check(Desc->CachedProperty);
-		check(Bag.GetValue().IsValid());
-		void* Address = Bag.GetMutableValue().GetMutableMemory() + Desc->CachedProperty->GetOffset_ForInternal();
-		return SetPropertyFromDouble(*Desc, Address, (double)Value);
+		if (Desc->ValueType != EPropertyBagPropertyType::Object
+			&& Desc->ValueType != EPropertyBagPropertyType::SoftObject
+			&& Desc->ValueType != EPropertyBagPropertyType::Class
+			&& Desc->ValueType != EPropertyBagPropertyType::SoftClass)
+		{
+			return EPropertyBagResult::TypeMismatch;
+		}
+		if (Address == nullptr)
+		{
+			return EPropertyBagResult::OutOfBounds;
+		}
+		if (Desc->ContainerType != EPropertyBagContainerType::None)
+		{
+			return EPropertyBagResult::TypeMismatch;
+		}
+		
+		const FObjectPropertyBase* ObjectProperty = CastFieldChecked<FObjectPropertyBase>(Desc->CachedProperty);
+		check(ObjectProperty->PropertyClass);
+		check(Desc->ValueTypeObject);
+
+		if (Desc->ValueType == EPropertyBagPropertyType::Object
+			|| Desc->ValueType == EPropertyBagPropertyType::SoftObject)
+		{
+			if (InValue && CanCastTo(InValue->GetClass(),  ObjectProperty->PropertyClass) == false)
+			{
+				return EPropertyBagResult::TypeMismatch;
+			}
+		}
+		else
+		{
+			const UClass* Class = Cast<UClass>(InValue);
+			const UClass* PropClass = nullptr;
+
+			if (const FClassProperty* ClassProperty = CastFieldChecked<FClassProperty>(Desc->CachedProperty))
+			{
+				PropClass = ClassProperty->MetaClass;
+			}
+			else if (const FSoftClassProperty* SoftClassProperty = CastFieldChecked<FSoftClassProperty>(Desc->CachedProperty))
+			{
+				PropClass = SoftClassProperty->MetaClass;
+			}
+			
+			if (!Class || !PropClass || !Class->IsChildOf(PropClass))
+			{
+				return EPropertyBagResult::TypeMismatch;
+			}
+		}
+
+		ObjectProperty->SetObjectPropertyValue(Address, InValue);
+
+		return EPropertyBagResult::Success;
 	}
 
-	template<typename T, typename PropT>
-	EPropertyBagResult SetValue(const FInstancedPropertyBag& Bag, const FName Name, const T& Value)
+	void CopyMatchingValuesByID(const FConstStructView Source, const FStructView Target)
 	{
-		const FPropertyBagPropertyDesc* Desc = Bag.FindPropertyDescByName(Name);
-		if (Desc == nullptr || Desc->CachedProperty == nullptr || !Desc->CachedProperty->IsA<PropT>())
+		if (!Source.IsValid() || !Target.IsValid())
 		{
-			return EPropertyBagResult::PropertyNotFound;
+			return;
+		}
+
+		const UPropertyBag* SourceBagStruct = Cast<const UPropertyBag>(Source.GetScriptStruct());
+		const UPropertyBag* TargetBagStruct = Cast<const UPropertyBag>(Target.GetScriptStruct());
+
+		if (!SourceBagStruct || !TargetBagStruct)
+		{
+			return;
+		}
+
+		// Iterate over source and copy to target if possible. Source is expected to usually have less items.
+		for (const FPropertyBagPropertyDesc& SourceDesc : SourceBagStruct->GetPropertyDescs())
+		{
+			const FPropertyBagPropertyDesc* PotentialTargetDesc = TargetBagStruct->FindPropertyDescByID(SourceDesc.ID);
+			if (PotentialTargetDesc == nullptr
+				|| PotentialTargetDesc->CachedProperty == nullptr
+				|| SourceDesc.CachedProperty == nullptr)
+			{
+				continue;
+			}
+
+			const FPropertyBagPropertyDesc& TargetDesc = *PotentialTargetDesc;
+			void* TargetAddress = Target.GetMutableMemory() + TargetDesc.CachedProperty->GetOffset_ForInternal();
+			const void* SourceAddress = Source.GetMemory() + SourceDesc.CachedProperty->GetOffset_ForInternal();
+			
+			if (TargetDesc.CompatibleType(SourceDesc))
+			{
+				TargetDesc.CachedProperty->CopyCompleteValue(TargetAddress, SourceAddress);
+			}
+			else if (TargetDesc.ContainerType == EPropertyBagContainerType::None
+					&& SourceDesc.ContainerType == EPropertyBagContainerType::None)
+			{
+				if (TargetDesc.IsNumericType() && SourceDesc.IsNumericType())
+				{
+					// Try to convert numeric types.
+					if (TargetDesc.IsNumericFloatType())
+					{
+						double Value = 0;
+						if (GetPropertyAsDouble(&SourceDesc, SourceAddress, Value) == EPropertyBagResult::Success)
+						{
+							SetPropertyFromDouble(&TargetDesc, TargetAddress, Value);
+						}
+					}
+					else
+					{
+						int64 Value = 0;
+						if (GetPropertyAsInt64(&SourceDesc, SourceAddress, Value) == EPropertyBagResult::Success)
+						{
+							SetPropertyFromInt64(&TargetDesc, TargetAddress, Value);
+						}
+					}
+				}
+				else if ((TargetDesc.IsObjectType() && SourceDesc.IsObjectType())
+					|| (TargetDesc.IsClassType() && SourceDesc.IsClassType()))
+				{
+					// Try convert between compatible objects and classes.
+					const UClass* TargetObjectClass = Cast<const UClass>(TargetDesc.ValueTypeObject);
+					const UClass* SourceObjectClass = Cast<const UClass>(SourceDesc.ValueTypeObject);
+					if (CanCastTo(SourceObjectClass, TargetObjectClass))
+					{
+						const FObjectPropertyBase* TargetProp = CastFieldChecked<FObjectPropertyBase>(TargetDesc.CachedProperty);
+						const FObjectPropertyBase* SourceProp = CastFieldChecked<FObjectPropertyBase>(SourceDesc.CachedProperty);
+						TargetProp->SetObjectPropertyValue(TargetAddress, SourceProp->GetObjectPropertyValue(SourceAddress));
+					}
+				}
+			}
 		}
-		check(Desc->CachedProperty);
-		const PropT* Property = CastFieldChecked<PropT>(Desc->CachedProperty);
-		check(Bag.GetValue().IsValid());
-		void* Address = Bag.GetMutableValue().GetMutableMemory() + Desc->CachedProperty->GetOffset_ForInternal();
-		Property->SetPropertyValue(Address, Value);
-		return EPropertyBagResult::Success;
 	}
 
 	void RemovePropertyByName(TArray<FPropertyBagPropertyDesc>& Descs, const FName PropertyName, const int32 StartIndex = 0)
@@ -689,7 +993,6 @@ namespace UE::StructUtils::Private
 		}
 	}
 
-
 }; // UE::StructUtils::Private
 
 
@@ -708,6 +1011,7 @@ FPropertyBagPropertyDesc::FPropertyBagPropertyDesc(const FName InName, const FPr
 {
 	ValueType = UE::StructUtils::Private::GetValueTypeFromProperty(InSourceProperty);
 	ValueTypeObject = UE::StructUtils::Private::GetValueTypeObjectFromProperty(InSourceProperty);
+	ContainerType = UE::StructUtils::Private::GetContainerTypeFromProperty(InSourceProperty);
 
 #if WITH_EDITORONLY_DATA
 	if (const TMap<FName, FString>* SourcePropertyMetaData = InSourceProperty->GetMetaDataMap())
@@ -727,6 +1031,11 @@ static FArchive& operator<<(FArchive& Ar, FPropertyBagPropertyDesc& Bag)
 	Ar << Bag.Name;
 	Ar << Bag.ValueType;
 
+	if (Ar.CustomVer(FPropertyBagCustomVersion::GUID) >= FPropertyBagCustomVersion::ContainerTypes)
+	{
+		Ar << Bag.ContainerType;
+	}
+
 	bool bHasMetaData = false;
 #if WITH_EDITORONLY_DATA
 	if (Ar.IsSaving())
@@ -795,6 +1104,12 @@ bool FPropertyBagPropertyDesc::IsClassType() const
 
 bool FPropertyBagPropertyDesc::CompatibleType(const FPropertyBagPropertyDesc& Other) const
 {
+	// Containers must match
+	if (ContainerType != Other.ContainerType)
+	{
+		return false;
+	}
+
 	// Values must match.
 	if (ValueType != Other.ValueType)
 	{
@@ -873,6 +1188,11 @@ void FInstancedPropertyBag::AddProperty(const FName InName, const EPropertyBagPr
 	AddProperties({ FPropertyBagPropertyDesc(InName, InValueType, InValueTypeObject) });
 }
 
+void FInstancedPropertyBag::AddContainerProperty(const FName InName, const EPropertyBagContainerType InContainerType, const EPropertyBagPropertyType InValueType, UObject* InValueTypeObject)
+{
+	AddProperties({ FPropertyBagPropertyDesc(InName, InContainerType, InValueType, InValueTypeObject) });
+}
+
 void FInstancedPropertyBag::AddProperty(const FName InName, const FProperty* InSourceProperty)
 {
 	AddProperties({ FPropertyBagPropertyDesc(InName, InSourceProperty ) });
@@ -941,290 +1261,253 @@ const FPropertyBagPropertyDesc* FInstancedPropertyBag::FindPropertyDescByName(co
 	return nullptr;
 }
 
+void* FInstancedPropertyBag::GetValueAddress(const FPropertyBagPropertyDesc* Desc) const
+{
+	if (Desc == nullptr || !Value.IsValid())
+	{
+		return nullptr;
+	}
+	return Value.GetMutableMemory() + Desc->CachedProperty->GetOffset_ForInternal();
+}
+
 TValueOrError<bool, EPropertyBagResult> FInstancedPropertyBag::GetValueBool(const FName Name) const
 {
-	TValueOrError<int64, EPropertyBagResult> Result = UE::StructUtils::Private::GetValueInt64<int64>(*this, Name);
-	if (!Result.IsValid())
+	const FPropertyBagPropertyDesc* Desc = FindPropertyDescByName(Name);
+	int64 ReturnValue = 0;
+	const EPropertyBagResult Result = UE::StructUtils::Private::GetPropertyAsInt64(Desc, GetValueAddress(Desc), ReturnValue);
+	if (Result != EPropertyBagResult::Success)
 	{
-		return MakeError(Result.GetError());
+		return MakeError(Result);
 	}
-	return MakeValue(Result.GetValue() != 0);
+	return MakeValue(ReturnValue != 0);
 }
 
 TValueOrError<uint8, EPropertyBagResult> FInstancedPropertyBag::GetValueByte(const FName Name) const
 {
-	return UE::StructUtils::Private::GetValueInt64<uint8>(*this, Name);
+	const FPropertyBagPropertyDesc* Desc = FindPropertyDescByName(Name);
+	int64 ReturnValue = 0;
+	const EPropertyBagResult Result = UE::StructUtils::Private::GetPropertyAsInt64(Desc, GetValueAddress(Desc), ReturnValue);
+	if (Result != EPropertyBagResult::Success)
+	{
+		return MakeError(Result);
+	}
+	return MakeValue((uint8)ReturnValue);
 }
 
 TValueOrError<int32, EPropertyBagResult> FInstancedPropertyBag::GetValueInt32(const FName Name) const
 {
-	return UE::StructUtils::Private::GetValueInt64<int32>(*this, Name);
+	const FPropertyBagPropertyDesc* Desc = FindPropertyDescByName(Name);
+	int64 ReturnValue = 0;
+	const EPropertyBagResult Result = UE::StructUtils::Private::GetPropertyAsInt64(Desc, GetValueAddress(Desc), ReturnValue);
+	if (Result != EPropertyBagResult::Success)
+	{
+		return MakeError(Result);
+	}
+	return MakeValue((int32)ReturnValue);
 }
 
 TValueOrError<int64, EPropertyBagResult> FInstancedPropertyBag::GetValueInt64(const FName Name) const
 {
-	return UE::StructUtils::Private::GetValueInt64<int64>(*this, Name);
+	const FPropertyBagPropertyDesc* Desc = FindPropertyDescByName(Name);
+	int64 ReturnValue = 0;
+	const EPropertyBagResult Result = UE::StructUtils::Private::GetPropertyAsInt64(Desc, GetValueAddress(Desc), ReturnValue);
+	if (Result != EPropertyBagResult::Success)
+	{
+		return MakeError(Result);
+	}
+	return MakeValue(ReturnValue);
 }
 
 TValueOrError<float, EPropertyBagResult> FInstancedPropertyBag::GetValueFloat(const FName Name) const
 {
-	return UE::StructUtils::Private::GetValueDouble<float>(*this, Name);
+	const FPropertyBagPropertyDesc* Desc = FindPropertyDescByName(Name);
+	double ReturnValue = 0;
+	const EPropertyBagResult Result = UE::StructUtils::Private::GetPropertyAsDouble(Desc, GetValueAddress(Desc), ReturnValue);
+	if (Result != EPropertyBagResult::Success)
+	{
+		return MakeError(Result);
+	}
+	return MakeValue((float)ReturnValue);
 }
 
 TValueOrError<double, EPropertyBagResult> FInstancedPropertyBag::GetValueDouble(const FName Name) const
 {
-	return UE::StructUtils::Private::GetValueDouble<double>(*this, Name);
+	const FPropertyBagPropertyDesc* Desc = FindPropertyDescByName(Name);
+	double ReturnValue = 0;
+	const EPropertyBagResult Result = UE::StructUtils::Private::GetPropertyAsDouble(Desc, GetValueAddress(Desc), ReturnValue);
+	if (Result != EPropertyBagResult::Success)
+	{
+		return MakeError(Result);
+	}
+	return MakeValue(ReturnValue);
 }
 
 TValueOrError<FName, EPropertyBagResult> FInstancedPropertyBag::GetValueName(const FName Name) const
 {
-	return UE::StructUtils::Private::GetValue<FName, FNameProperty>(*this, Name);
+	const FPropertyBagPropertyDesc* Desc = FindPropertyDescByName(Name);
+	FName ReturnValue;
+	const EPropertyBagResult Result = UE::StructUtils::Private::GetPropertyValue<FName, FNameProperty>(Desc, GetValueAddress(Desc), ReturnValue);
+	if (Result != EPropertyBagResult::Success)
+	{
+		return MakeError(Result);
+	}
+	return MakeValue(ReturnValue);
 }
 
 TValueOrError<FString, EPropertyBagResult> FInstancedPropertyBag::GetValueString(const FName Name) const
 {
-	return UE::StructUtils::Private::GetValue<FString, FStrProperty>(*this, Name);
+	const FPropertyBagPropertyDesc* Desc = FindPropertyDescByName(Name);
+	FString ReturnValue;
+	const EPropertyBagResult Result = UE::StructUtils::Private::GetPropertyValue<FString, FStrProperty>(Desc, GetValueAddress(Desc), ReturnValue);
+	if (Result != EPropertyBagResult::Success)
+	{
+		return MakeError(Result);
+	}
+	return MakeValue(ReturnValue);
 }
 
 TValueOrError<FText, EPropertyBagResult> FInstancedPropertyBag::GetValueText(const FName Name) const
 {
-	return UE::StructUtils::Private::GetValue<FText, FTextProperty>(*this, Name);
+	const FPropertyBagPropertyDesc* Desc = FindPropertyDescByName(Name);
+	FText ReturnValue;
+	const EPropertyBagResult Result = UE::StructUtils::Private::GetPropertyValue<FText, FTextProperty>(Desc, GetValueAddress(Desc), ReturnValue);
+	if (Result != EPropertyBagResult::Success)
+	{
+		return MakeError(Result);
+	}
+	return MakeValue(ReturnValue);
 }
 
 TValueOrError<uint8, EPropertyBagResult> FInstancedPropertyBag::GetValueEnum(const FName Name, const UEnum* RequestedEnum) const
 {
 	const FPropertyBagPropertyDesc* Desc = FindPropertyDescByName(Name);
-	if (Desc == nullptr || Desc->ValueType != EPropertyBagPropertyType::Enum)
-	{
-		return MakeError(Desc == nullptr ? EPropertyBagResult::PropertyNotFound : EPropertyBagResult::TypeMismatch);
-	}
-	check(Desc->CachedProperty);
-	const FEnumProperty* EnumProperty = CastFieldChecked<FEnumProperty>(Desc->CachedProperty);
-	const FNumericProperty* UnderlyingProperty = EnumProperty->GetUnderlyingProperty();
-	check(UnderlyingProperty);
-	
-	if (RequestedEnum != EnumProperty->GetEnum())
+	uint8 ReturnValue;
+	const EPropertyBagResult Result = UE::StructUtils::Private::GetPropertyValueAsEnum(Desc, RequestedEnum, GetValueAddress(Desc), ReturnValue);
+	if (Result != EPropertyBagResult::Success)
 	{
-		return MakeError(EPropertyBagResult::TypeMismatch);
+		return MakeError(Result);
 	}
-	
-	check(Value.IsValid());
-	const void* Address = Value.GetMutableMemory() + Desc->CachedProperty->GetOffset_ForInternal();
-	return MakeValue((uint8)UnderlyingProperty->GetUnsignedIntPropertyValue(Address));
+	return MakeValue(ReturnValue);
 }
 
 TValueOrError<FStructView, EPropertyBagResult> FInstancedPropertyBag::GetValueStruct(const FName Name, const UScriptStruct* RequestedStruct) const
 {
 	const FPropertyBagPropertyDesc* Desc = FindPropertyDescByName(Name);
-	if (Desc == nullptr || Desc->ValueType != EPropertyBagPropertyType::Struct)
-	{
-		return MakeError(Desc == nullptr ? EPropertyBagResult::PropertyNotFound : EPropertyBagResult::TypeMismatch);
-	}
-	check(Desc->CachedProperty);
-	const FStructProperty* StructProperty = CastFieldChecked<FStructProperty>(Desc->CachedProperty);
-	check(StructProperty->Struct);
-
-	if (RequestedStruct != nullptr && UE::StructUtils::Private::CanCastTo(StructProperty->Struct, RequestedStruct) == false)
+	FStructView ReturnValue;
+	const EPropertyBagResult Result = UE::StructUtils::Private::GetPropertyValueAsStruct(Desc, RequestedStruct, GetValueAddress(Desc), ReturnValue);
+	if (Result != EPropertyBagResult::Success)
 	{
-		return MakeError(EPropertyBagResult::TypeMismatch);
+		return MakeError(Result);
 	}
-	
-	check(Value.IsValid());
-	void* Address = Value.GetMutableMemory() + Desc->CachedProperty->GetOffset_ForInternal();
-	return MakeValue(FStructView(StructProperty->Struct, (uint8*)Address));
+	return MakeValue(ReturnValue);
 }
 
 TValueOrError<UObject*, EPropertyBagResult> FInstancedPropertyBag::GetValueObject(const FName Name, const UClass* RequestedClass) const
 {
 	const FPropertyBagPropertyDesc* Desc = FindPropertyDescByName(Name);
-	if (Desc == nullptr || (Desc->ValueType != EPropertyBagPropertyType::Object && Desc->ValueType != EPropertyBagPropertyType::SoftObject))
-	{
-		return MakeError(Desc == nullptr ? EPropertyBagResult::PropertyNotFound : EPropertyBagResult::TypeMismatch);
-	}
-	check(Desc->CachedProperty);
-	const FObjectPropertyBase* ObjectProperty = CastFieldChecked<FObjectPropertyBase>(Desc->CachedProperty);
-	check(ObjectProperty->PropertyClass);
-
-	if (RequestedClass != nullptr && UE::StructUtils::Private::CanCastTo(ObjectProperty->PropertyClass, RequestedClass) == false)
+	UObject* ReturnValue = nullptr;
+	const EPropertyBagResult Result = UE::StructUtils::Private::GetPropertyValueAsObject(Desc, RequestedClass, GetValueAddress(Desc), ReturnValue);
+	if (Result != EPropertyBagResult::Success)
 	{
-		return MakeError(EPropertyBagResult::TypeMismatch);
+		return MakeError(Result);
 	}
-	
-	check(Value.IsValid());
-	const void* Address = Value.GetMemory() + Desc->CachedProperty->GetOffset_ForInternal();
-	return MakeValue(ObjectProperty->GetObjectPropertyValue(Address));
+	return MakeValue(ReturnValue);
 }
 
 TValueOrError<UClass*, EPropertyBagResult> FInstancedPropertyBag::GetValueClass(const FName Name) const
 {
 	const FPropertyBagPropertyDesc* Desc = FindPropertyDescByName(Name);
-	if (Desc == nullptr || (Desc->ValueType != EPropertyBagPropertyType::Class && Desc->ValueType != EPropertyBagPropertyType::SoftClass))
+	UObject* ReturnValue = nullptr;
+	const EPropertyBagResult Result = UE::StructUtils::Private::GetPropertyValueAsObject(Desc, nullptr, GetValueAddress(Desc), ReturnValue);
+	if (Result != EPropertyBagResult::Success)
 	{
-		return MakeError(Desc == nullptr ? EPropertyBagResult::PropertyNotFound : EPropertyBagResult::TypeMismatch);
+		return MakeError(Result);
 	}
-	check(Desc->CachedProperty);
-	const FObjectPropertyBase* ObjectProperty = CastFieldChecked<FObjectPropertyBase>(Desc->CachedProperty);
-	check(Value.IsValid());
-	const void* Address = Value.GetMemory() + Desc->CachedProperty->GetOffset_ForInternal();
-	return MakeValue(CastChecked<UClass>(ObjectProperty->GetObjectPropertyValue(Address)));
+	UClass* Class = Cast<UClass>(ReturnValue);
+	if (Class == nullptr && ReturnValue != nullptr)
+	{
+		return MakeError(EPropertyBagResult::TypeMismatch);
+	}
+	return MakeValue(Class);
 }
 
 
 EPropertyBagResult FInstancedPropertyBag::SetValueBool(const FName Name, const bool bInValue) const
 {
-	return UE::StructUtils::Private::SetValueInt64<int64>(*this, Name, bInValue ? 1 : 0);
+	const FPropertyBagPropertyDesc* Desc = FindPropertyDescByName(Name);
+	return UE::StructUtils::Private::SetPropertyFromInt64(Desc, GetValueAddress(Desc), bInValue ? 1 : 0);
 }
 
 EPropertyBagResult FInstancedPropertyBag::SetValueByte(const FName Name, const uint8 InValue) const
 {
-	return UE::StructUtils::Private::SetValueInt64<uint8>(*this, Name, InValue);
+	const FPropertyBagPropertyDesc* Desc = FindPropertyDescByName(Name);
+	return UE::StructUtils::Private::SetPropertyFromInt64(Desc, GetValueAddress(Desc), InValue);
 }
 
 EPropertyBagResult FInstancedPropertyBag::SetValueInt32(const FName Name, const int32 InValue) const
 {
-	return UE::StructUtils::Private::SetValueInt64<int32>(*this, Name, InValue);
+	const FPropertyBagPropertyDesc* Desc = FindPropertyDescByName(Name);
+	return UE::StructUtils::Private::SetPropertyFromInt64(Desc, GetValueAddress(Desc), InValue);
 }
 
 EPropertyBagResult FInstancedPropertyBag::SetValueInt64(const FName Name, const int64 InValue) const
 {
-	return UE::StructUtils::Private::SetValueInt64<int64>(*this, Name, InValue);
+	const FPropertyBagPropertyDesc* Desc = FindPropertyDescByName(Name);
+	return UE::StructUtils::Private::SetPropertyFromInt64(Desc, GetValueAddress(Desc), InValue);
 }
 
 EPropertyBagResult FInstancedPropertyBag::SetValueFloat(const FName Name, const float InValue) const
 {
-	return UE::StructUtils::Private::SetValueDouble<float>(*this, Name, InValue);
+	const FPropertyBagPropertyDesc* Desc = FindPropertyDescByName(Name);
+	return UE::StructUtils::Private::SetPropertyFromDouble(Desc, GetValueAddress(Desc), InValue);
 }
 
 EPropertyBagResult FInstancedPropertyBag::SetValueDouble(const FName Name, const double InValue) const
 {
-	return UE::StructUtils::Private::SetValueDouble<double>(*this, Name, InValue);
+	const FPropertyBagPropertyDesc* Desc = FindPropertyDescByName(Name);
+	return UE::StructUtils::Private::SetPropertyFromDouble(Desc, GetValueAddress(Desc), InValue);
 }
 
 EPropertyBagResult FInstancedPropertyBag::SetValueName(const FName Name, const FName InValue) const
 {
-	return UE::StructUtils::Private::SetValue<FName, FNameProperty>(*this, Name, InValue);
+	const FPropertyBagPropertyDesc* Desc = FindPropertyDescByName(Name);
+	return UE::StructUtils::Private::SetPropertyValue<FName, FNameProperty>(Desc, GetValueAddress(Desc), InValue);
 }
 
 EPropertyBagResult FInstancedPropertyBag::SetValueString(const FName Name, const FString& InValue) const
 {
-	return UE::StructUtils::Private::SetValue<FString, FStrProperty>(*this, Name, InValue);
+	const FPropertyBagPropertyDesc* Desc = FindPropertyDescByName(Name);
+	return UE::StructUtils::Private::SetPropertyValue<FString, FStrProperty>(Desc, GetValueAddress(Desc), InValue);
 }
 
 EPropertyBagResult FInstancedPropertyBag::SetValueText(const FName Name, const FText& InValue) const
 {
-	return UE::StructUtils::Private::SetValue<FText, FTextProperty>(*this, Name, InValue);
+	const FPropertyBagPropertyDesc* Desc = FindPropertyDescByName(Name);
+	return UE::StructUtils::Private::SetPropertyValue<FText, FTextProperty>(Desc, GetValueAddress(Desc), InValue);
 }
 
 EPropertyBagResult FInstancedPropertyBag::SetValueEnum(const FName Name, const uint8 InValue, const UEnum* Enum) const
 {
 	const FPropertyBagPropertyDesc* Desc = FindPropertyDescByName(Name);
-	if (Desc == nullptr || Desc->ValueType != EPropertyBagPropertyType::Enum)
-	{
-		return Desc == nullptr ? EPropertyBagResult::PropertyNotFound : EPropertyBagResult::TypeMismatch;
-	}
-	check(Desc->CachedProperty);
-	const FEnumProperty* EnumProperty = CastFieldChecked<FEnumProperty>(Desc->CachedProperty);
-	const FNumericProperty* UnderlyingProperty = EnumProperty->GetUnderlyingProperty();
-	check(UnderlyingProperty);
-
-	if (Enum != EnumProperty->GetEnum())
-	{
-		return EPropertyBagResult::TypeMismatch;
-	}
-	
-	check(Value.IsValid());
-	void* Address = Value.GetMutableMemory() + Desc->CachedProperty->GetOffset_ForInternal();
-	UnderlyingProperty->SetIntPropertyValue(Address, (uint64)InValue);
-	
-	return EPropertyBagResult::Success;
+	return UE::StructUtils::Private::SetPropertyValueAsEnum(Desc, GetValueAddress(Desc), InValue, Enum);
 }
 
 EPropertyBagResult FInstancedPropertyBag::SetValueStruct(const FName Name, FConstStructView InValue) const
 {
 	const FPropertyBagPropertyDesc* Desc = FindPropertyDescByName(Name);
-	if (Desc == nullptr || Desc->ValueType != EPropertyBagPropertyType::Struct)
-	{
-		return Desc == nullptr ? EPropertyBagResult::PropertyNotFound : EPropertyBagResult::TypeMismatch;
-	}
-	check(Desc->CachedProperty);
-	const FStructProperty* StructProperty = CastFieldChecked<FStructProperty>(Desc->CachedProperty);
-	check(StructProperty->Struct);
-
-	if (InValue.GetScriptStruct() && InValue.GetScriptStruct() != StructProperty->Struct)
-	{
-		return EPropertyBagResult::TypeMismatch;
-	}
-
-	check(Value.IsValid());
-	void* Address = Value.GetMutableMemory() + Desc->CachedProperty->GetOffset_ForInternal();
-
-	if (InValue.IsValid())
-	{
-		StructProperty->Struct->CopyScriptStruct(Address, InValue.GetMemory());
-	}
-	else
-	{
-		StructProperty->Struct->ClearScriptStruct(Address);
-	}
-
-	return EPropertyBagResult::Success;
+	return UE::StructUtils::Private::SetPropertyValueAsStruct(Desc, GetValueAddress(Desc), InValue);
 }
 
 EPropertyBagResult FInstancedPropertyBag::SetValueObject(const FName Name, UObject* InValue) const
 {
 	const FPropertyBagPropertyDesc* Desc = FindPropertyDescByName(Name);
-	if (Desc == nullptr || (Desc->ValueType != EPropertyBagPropertyType::Object && Desc->ValueType != EPropertyBagPropertyType::SoftObject))
-	{
-		return Desc == nullptr ? EPropertyBagResult::PropertyNotFound : EPropertyBagResult::TypeMismatch;
-	}
-	check(Desc->CachedProperty);
-	const FObjectPropertyBase* ObjectProperty = CastFieldChecked<FObjectPropertyBase>(Desc->CachedProperty);
-	check(ObjectProperty->PropertyClass);
-
-	if (InValue && UE::StructUtils::Private::CanCastTo(InValue->GetClass(), ObjectProperty->PropertyClass) == false)
-	{
-		return EPropertyBagResult::TypeMismatch;
-	}
-	
-	check(Value.IsValid());
-	void* Address = Value.GetMutableMemory() + Desc->CachedProperty->GetOffset_ForInternal();
-	ObjectProperty->SetObjectPropertyValue(Address, InValue);
-
-	return EPropertyBagResult::Success;
+	return UE::StructUtils::Private::SetPropertyValueAsObject(Desc, GetValueAddress(Desc), InValue);
 }
 
 EPropertyBagResult FInstancedPropertyBag::SetValueClass(const FName Name, UClass* InValue) const
 {
 	const FPropertyBagPropertyDesc* Desc = FindPropertyDescByName(Name);
-	if (Desc == nullptr || (Desc->ValueType != EPropertyBagPropertyType::Class && Desc->ValueType != EPropertyBagPropertyType::SoftClass))
-	{
-		return Desc == nullptr ? EPropertyBagResult::PropertyNotFound : EPropertyBagResult::TypeMismatch;
-	}
-	check(Desc->CachedProperty);
-	
-	check(Value.IsValid());
-	void* Address = Value.GetMutableMemory() + Desc->CachedProperty->GetOffset_ForInternal();
-
-	if (Desc->ValueType == EPropertyBagPropertyType::Class)
-	{
-		const FClassProperty* ClassProperty = CastFieldChecked<FClassProperty>(Desc->CachedProperty);
-		if (InValue && InValue->IsChildOf(ClassProperty->MetaClass) == false)
-		{
-			return EPropertyBagResult::TypeMismatch;
-		}
-		ClassProperty->SetObjectPropertyValue(Address, InValue);
-	}
-	else
-	{
-		const FSoftClassProperty* ClassProperty = CastFieldChecked<FSoftClassProperty>(Desc->CachedProperty);
-		if (InValue && InValue->IsChildOf(ClassProperty->MetaClass) == false)
-		{
-			return EPropertyBagResult::TypeMismatch;
-		}
-		ClassProperty->SetObjectPropertyValue(Address, InValue);
-	}
-
-	return EPropertyBagResult::Success;
+	return UE::StructUtils::Private::SetPropertyValueAsObject(Desc, GetValueAddress(Desc), InValue);
 }
 
 EPropertyBagResult FInstancedPropertyBag::SetValue(const FName Name, const FProperty* InSourceProperty, const void* InSourceContainerAddress) const
@@ -1251,8 +1534,32 @@ EPropertyBagResult FInstancedPropertyBag::SetValue(const FName Name, const FProp
 	return EPropertyBagResult::Success;
 }
 
+TValueOrError<FPropertyBagArrayRef, EPropertyBagResult> FInstancedPropertyBag::GetArrayRef(const FName Name) const
+{
+	const FPropertyBagPropertyDesc* Desc = FindPropertyDescByName(Name);
+	if (Desc == nullptr)
+	{
+		return MakeError(EPropertyBagResult::PropertyNotFound);
+	}
+	check(Desc->CachedProperty);
+
+	if (Desc->ContainerType != EPropertyBagContainerType::Array)
+	{
+		return MakeError(EPropertyBagResult::TypeMismatch);
+	}
+
+	const void* Address = GetValueAddress(Desc);
+	if (Address == nullptr)
+	{
+		return MakeError(EPropertyBagResult::PropertyNotFound);
+	}
+	
+	return MakeValue(FPropertyBagArrayRef(*Desc, Address));
+}
+
 bool FInstancedPropertyBag::Serialize(FArchive& Ar)
 {
+	// Obsolete, use custom version instead.
 	enum class EVersion : uint8
 	{
 		InitialVersion = 0,
@@ -1261,11 +1568,15 @@ bool FInstancedPropertyBag::Serialize(FArchive& Ar)
 		VersionPlusOne,
 		LatestVersion = VersionPlusOne - 1
 	};
-
 	EVersion Version = EVersion::LatestVersion;
 
-	Ar << Version;
+	Ar.UsingCustomVersion(FPropertyBagCustomVersion::GUID);
 
+	if (Ar.CustomVer(FPropertyBagCustomVersion::GUID) < FPropertyBagCustomVersion::ContainerTypes)
+	{
+		Ar << Version;
+	}
+	
 	UPropertyBag* BagStruct = Cast<UPropertyBag>(const_cast<UScriptStruct*>(Value.GetScriptStruct()));
 	bool bHasData = (BagStruct != nullptr);
 	
@@ -1328,6 +1639,226 @@ bool FInstancedPropertyBag::Serialize(FArchive& Ar)
 	return true;
 }
 
+
+//----------------------------------------------------------------//
+//  FPropertyBagArrayRef
+//----------------------------------------------------------------//
+
+TValueOrError<bool, EPropertyBagResult> FPropertyBagArrayRef::GetValueBool(const int32 Index) const
+{
+	int64 IntValue = 0;
+	const EPropertyBagResult Result = UE::StructUtils::Private::GetPropertyAsInt64(&ValueDesc, GetAddress(Index), IntValue);
+	if (Result != EPropertyBagResult::Success)
+	{
+		return MakeError(Result);
+	}
+	return MakeValue(IntValue != 0);
+}
+
+TValueOrError<uint8, EPropertyBagResult> FPropertyBagArrayRef::GetValueByte(const int32 Index) const
+{
+	int64 IntValue = 0;
+	const EPropertyBagResult Result = UE::StructUtils::Private::GetPropertyAsInt64(&ValueDesc, GetAddress(Index), IntValue);
+	if (Result != EPropertyBagResult::Success)
+	{
+		return MakeError(Result);
+	}
+	return MakeValue((uint8)IntValue);
+}
+
+TValueOrError<int32, EPropertyBagResult> FPropertyBagArrayRef::GetValueInt32(const int32 Index) const
+{
+	int64 IntValue = 0;
+	const EPropertyBagResult Result = UE::StructUtils::Private::GetPropertyAsInt64(&ValueDesc, GetAddress(Index), IntValue);
+	if (Result != EPropertyBagResult::Success)
+	{
+		return MakeError(Result);
+	}
+	return MakeValue((int32)IntValue);
+}
+
+TValueOrError<int64, EPropertyBagResult> FPropertyBagArrayRef::GetValueInt64(const int32 Index) const
+{
+	int64 IntValue = 0;
+	const EPropertyBagResult Result = UE::StructUtils::Private::GetPropertyAsInt64(&ValueDesc, GetAddress(Index), IntValue);
+	if (Result != EPropertyBagResult::Success)
+	{
+		return MakeError(Result);
+	}
+	return MakeValue(IntValue);
+}
+
+TValueOrError<float, EPropertyBagResult> FPropertyBagArrayRef::GetValueFloat(const int32 Index) const
+{
+	double DblValue = 0;
+	const EPropertyBagResult Result = UE::StructUtils::Private::GetPropertyAsDouble(&ValueDesc, GetAddress(Index), DblValue);
+	if (Result != EPropertyBagResult::Success)
+	{
+		return MakeError(Result);
+	}
+	return MakeValue((float)DblValue);
+}
+
+TValueOrError<double, EPropertyBagResult> FPropertyBagArrayRef::GetValueDouble(const int32 Index) const
+{
+	double DblValue = 0;
+	const EPropertyBagResult Result = UE::StructUtils::Private::GetPropertyAsDouble(&ValueDesc, GetAddress(Index), DblValue);
+	if (Result != EPropertyBagResult::Success)
+	{
+		return MakeError(Result);
+	}
+	return MakeValue(DblValue);
+}
+
+TValueOrError<FName, EPropertyBagResult> FPropertyBagArrayRef::GetValueName(const int32 Index) const
+{
+	FName Value;
+	const EPropertyBagResult Result = UE::StructUtils::Private::GetPropertyValue<FName, FNameProperty>(&ValueDesc, GetAddress(Index), Value);
+	if (Result != EPropertyBagResult::Success)
+	{
+		return MakeError(Result);
+	}
+	return MakeValue(Value);
+}
+
+TValueOrError<FString, EPropertyBagResult> FPropertyBagArrayRef::GetValueString(const int32 Index) const
+{
+	FString Value;
+	const EPropertyBagResult Result = UE::StructUtils::Private::GetPropertyValue<FString, FStrProperty>(&ValueDesc, GetAddress(Index), Value);
+	if (Result != EPropertyBagResult::Success)
+	{
+		return MakeError(Result);
+	}
+	return MakeValue(Value);
+}
+
+TValueOrError<FText, EPropertyBagResult> FPropertyBagArrayRef::GetValueText(const int32 Index) const
+{
+	FText Value;
+	const EPropertyBagResult Result = UE::StructUtils::Private::GetPropertyValue<FText, FTextProperty>(&ValueDesc, GetAddress(Index), Value);
+	if (Result != EPropertyBagResult::Success)
+	{
+		return MakeError(Result);
+	}
+	return MakeValue(Value);
+}
+
+TValueOrError<uint8, EPropertyBagResult> FPropertyBagArrayRef::GetValueEnum(const int32 Index, const UEnum* RequestedEnum) const
+{
+	uint8 Value;
+	const EPropertyBagResult Result = UE::StructUtils::Private::GetPropertyValueAsEnum(&ValueDesc, RequestedEnum, GetAddress(Index), Value);
+	if (Result != EPropertyBagResult::Success)
+	{
+		return MakeError(Result);
+	}
+	return MakeValue(Value);
+}
+
+TValueOrError<FStructView, EPropertyBagResult> FPropertyBagArrayRef::GetValueStruct(const int32 Index, const UScriptStruct* RequestedStruct) const
+{
+	FStructView Value;
+	const EPropertyBagResult Result = UE::StructUtils::Private::GetPropertyValueAsStruct(&ValueDesc, RequestedStruct, GetAddress(Index), Value);
+	if (Result != EPropertyBagResult::Success)
+	{
+		return MakeError(Result);
+	}
+	return MakeValue(Value);
+}
+
+TValueOrError<UObject*, EPropertyBagResult> FPropertyBagArrayRef::GetValueObject(const int32 Index, const UClass* RequestedClass) const
+{
+	UObject* Value = nullptr;
+	const EPropertyBagResult Result = UE::StructUtils::Private::GetPropertyValueAsObject(&ValueDesc, RequestedClass, GetAddress(Index), Value);
+	if (Result != EPropertyBagResult::Success)
+	{
+		return MakeError(Result);
+	}
+	return MakeValue(Value);
+}
+
+TValueOrError<UClass*, EPropertyBagResult> FPropertyBagArrayRef::GetValueClass(const int32 Index) const
+{
+	UObject* Value = nullptr;
+	const EPropertyBagResult Result = UE::StructUtils::Private::GetPropertyValueAsObject(&ValueDesc, nullptr, GetAddress(Index), Value);
+	if (Result != EPropertyBagResult::Success)
+	{
+		return MakeError(Result);
+	}
+	UClass* Class = Cast<UClass>(Value);
+	if (Class == nullptr && Value != nullptr)
+	{
+		return MakeError(EPropertyBagResult::TypeMismatch);
+	}
+	return MakeValue(Class);
+}
+
+
+EPropertyBagResult FPropertyBagArrayRef::SetValueBool(const int32 Index, const bool bInValue) const
+{
+	return UE::StructUtils::Private::SetPropertyFromInt64(&ValueDesc, GetMutableAddress(Index), bInValue ? 1 : 0);
+}
+
+EPropertyBagResult FPropertyBagArrayRef::SetValueByte(const int32 Index, const uint8 InValue) const
+{
+	return UE::StructUtils::Private::SetPropertyFromInt64(&ValueDesc, GetMutableAddress(Index), InValue);
+}
+
+EPropertyBagResult FPropertyBagArrayRef::SetValueInt32(const int32 Index, const int32 InValue) const
+{
+	return UE::StructUtils::Private::SetPropertyFromInt64(&ValueDesc, GetMutableAddress(Index), InValue);
+}
+
+EPropertyBagResult FPropertyBagArrayRef::SetValueInt64(const int32 Index, const int64 InValue) const
+{
+	return UE::StructUtils::Private::SetPropertyFromInt64(&ValueDesc, GetMutableAddress(Index), InValue);
+}
+
+EPropertyBagResult FPropertyBagArrayRef::SetValueFloat(const int32 Index, const float InValue) const
+{
+	return UE::StructUtils::Private::SetPropertyFromDouble(&ValueDesc, GetMutableAddress(Index), InValue);
+}
+
+EPropertyBagResult FPropertyBagArrayRef::SetValueDouble(const int32 Index, const double InValue) const
+{
+	return UE::StructUtils::Private::SetPropertyFromDouble(&ValueDesc, GetMutableAddress(Index), InValue);
+}
+
+EPropertyBagResult FPropertyBagArrayRef::SetValueName(const int32 Index, const FName InValue) const
+{
+	return UE::StructUtils::Private::SetPropertyValue<FName, FNameProperty>(&ValueDesc, GetMutableAddress(Index), InValue);
+}
+
+EPropertyBagResult FPropertyBagArrayRef::SetValueString(const int32 Index, const FString& InValue) const
+{
+	return UE::StructUtils::Private::SetPropertyValue<FString, FStrProperty>(&ValueDesc, GetMutableAddress(Index), InValue);
+}
+
+EPropertyBagResult FPropertyBagArrayRef::SetValueText(const int32 Index, const FText& InValue) const
+{
+	return UE::StructUtils::Private::SetPropertyValue<FText, FTextProperty>(&ValueDesc, GetMutableAddress(Index), InValue);
+}
+
+EPropertyBagResult FPropertyBagArrayRef::SetValueEnum(const int32 Index, const uint8 InValue, const UEnum* Enum) const
+{
+	return UE::StructUtils::Private::SetPropertyValueAsEnum(&ValueDesc, GetMutableAddress(Index), InValue, Enum);
+}
+
+EPropertyBagResult FPropertyBagArrayRef::SetValueStruct(const int32 Index, FConstStructView InValue) const
+{
+	return UE::StructUtils::Private::SetPropertyValueAsStruct(&ValueDesc, GetMutableAddress(Index), InValue);
+}
+
+EPropertyBagResult FPropertyBagArrayRef::SetValueObject(const int32 Index, UObject* InValue) const
+{
+	return UE::StructUtils::Private::SetPropertyValueAsObject(&ValueDesc, GetMutableAddress(Index), InValue);
+}
+
+EPropertyBagResult FPropertyBagArrayRef::SetValueClass(const int32 Index, UClass* InValue) const
+{
+	return UE::StructUtils::Private::SetPropertyValueAsObject(&ValueDesc, GetMutableAddress(Index), InValue);
+}
+
+
 //----------------------------------------------------------------//
 //  UPropertyBag
 //----------------------------------------------------------------//
diff --git a/Engine/Plugins/Experimental/StructUtils/Source/StructUtils/Public/PropertyBag.h b/Engine/Plugins/Experimental/StructUtils/Source/StructUtils/Public/PropertyBag.h
--- a/Engine/Plugins/Experimental/StructUtils/Source/StructUtils/Public/PropertyBag.h
+++ b/Engine/Plugins/Experimental/StructUtils/Source/StructUtils/Public/PropertyBag.h
@@ -5,6 +5,7 @@
 #include "InstancedStruct.h"
 #include "StructView.h"
 #include "Templates/ValueOrError.h"
+#include "UObject/UnrealType.h"
 #include <atomic>
 #include "PropertyBag.generated.h"
 
@@ -30,12 +31,21 @@ enum class EPropertyBagPropertyType : uint8
 	SoftClass UMETA(Hidden),
 };
 
+/** Property bag property container type. */
+UENUM()
+enum class EPropertyBagContainerType : uint8
+{
+	None UMETA(Hidden),
+	Array,
+};
+
 /** Getter and setter result code. */
 UENUM()
 enum class EPropertyBagResult : uint8
 {
 	Success,			// Operation succeeded.
 	TypeMismatch,		// Tried to access mismatching type (e.g. setting a struct to bool)
+	OutOfBounds,		// Tried to access an array property out of bounds.
 	PropertyNotFound,	// Could not find property of specified name.
 };
 
@@ -100,6 +110,13 @@ struct STRUCTUTILS_API FPropertyBagPropertyDesc
 		, ValueType(InValueType)
 	{
 	}
+	FPropertyBagPropertyDesc(const FName InName, const EPropertyBagContainerType InContainerType, const EPropertyBagPropertyType InValueType, UObject* InValueTypeObject = nullptr)
+		: ValueTypeObject(InValueTypeObject)
+		, Name(InName)
+		, ValueType(InValueType)
+		, ContainerType(InContainerType)
+	{
+	}
 
 	/** @return true if the two descriptors have the same type. Object types are compatible if Other can be cast to this type. */
 	bool CompatibleType(const FPropertyBagPropertyDesc& Other) const;
@@ -132,6 +149,10 @@ struct STRUCTUTILS_API FPropertyBagPropertyDesc
 	UPROPERTY(EditAnywhere, Category="Default")
 	EPropertyBagPropertyType ValueType = EPropertyBagPropertyType::None;
 
+	/** Type of the container described by this property. */
+	UPROPERTY(EditAnywhere, Category="Default")
+	EPropertyBagContainerType ContainerType = EPropertyBagContainerType::None;
+
 #if WITH_EDITORONLY_DATA
 	/** Editor-only meta data for CachedProperty */
 	UPROPERTY(EditAnywhere, Category="Default")
@@ -185,6 +206,7 @@ struct STRUCTUTILS_API FPropertyBagPropertyDesc
  */
 
 class UPropertyBag;
+class FPropertyBagArrayRef;
 
 USTRUCT()
 struct STRUCTUTILS_API FInstancedPropertyBag
@@ -233,10 +255,21 @@ struct STRUCTUTILS_API FInstancedPropertyBag
 	/**
 	 * Adds a new property to the bag. If property of same name already exists, it will be replaced with the new type.
 	 * Numeric property values will be converted if possible, when a property's type changes.
-	 * @param Descs Descriptors of new properties to add.
+	 * @param InName Name of the new property
+	 * @param InValueType Type of the new property
+	 * @param InValueTypeObject Type object (for struct, class, enum) of the new property
 	 */
 	void AddProperty(const FName InName, const EPropertyBagPropertyType InValueType, UObject* InValueTypeObject = nullptr);
 
+	/**
+	 * Adds a new container property to the bag. If property of same name already exists, it will be replaced with the new type.
+	 * @param InName Name of the new property
+	 * @param InContainerType Type of the new container
+	 * @param InValueType Type of the new property
+	 * @param InValueTypeObject Type object (for struct, class, enum) of the new property
+	 */
+	void AddContainerProperty(const FName InName, const EPropertyBagContainerType InContainerType, const EPropertyBagPropertyType InValueType, UObject* InValueTypeObject = nullptr);
+
 	/**
 	 * Adds a new property to the bag. Property type duplicated from source property to. If property of same name already exists, it will be replaced with the new type.
 	 * @param Descs Descriptors of new properties to add.
@@ -400,9 +433,18 @@ struct STRUCTUTILS_API FInstancedPropertyBag
 	 */
 	EPropertyBagResult SetValue(const FName Name, const FProperty* InSourceProperty, const void* InSourceContainerAddress) const;
 
+	/**
+	 * Returns helper class to modify and access an array property.
+	 * Note: Note: The array reference is not valid after the layout of the referenced property bag has changed!
+	 * @returns helper class to modify and access arrays
+	*/
+	TValueOrError<FPropertyBagArrayRef, EPropertyBagResult> GetArrayRef(const FName Name) const;
+
 	bool Serialize(FArchive& Ar);
 
 protected:
+	void* GetValueAddress(const FPropertyBagPropertyDesc* Desc) const;
+	
 	UPROPERTY(EditAnywhere, Category="")
 	FInstancedStruct Value;
 };
@@ -415,6 +457,179 @@ template<> struct TStructOpsTypeTraits<FInstancedPropertyBag> : public TStructOp
 	};
 };
 
+
+/**
+ * A reference to an array in FInstancedPropertyBag
+ * Allows to modify the array via the FScriptArrayHelper API, and contains helper methods to get and set properties.
+ *
+ *		FInstancedPropertyBag Bag;
+ *		Bag.AddProperties({
+ *			{ ArrayName, EPropertyBagContainerType::Array, EPropertyBagPropertyType::Float }
+ *		});
+ *
+ *		if (auto FloatArrayRes = Bag.GetArrayRef(ArrayName); FloatArrayRes.IsValid())
+ *		{
+ *			FPropertyBagArrayRef& FloatArray = FloatArrayRes.GetValue();
+ *			const int32 NewIndex = FloatArray.AddValue();
+ *			FloatArray.SetValueFloat(NewIndex, 123.0f);
+ *		}
+ * 
+ * Note: The array reference is not valid after the layout of the referenced property bag has changed! 
+ */
+class STRUCTUTILS_API FPropertyBagArrayRef : public FScriptArrayHelper
+{
+	FPropertyBagPropertyDesc ValueDesc;
+
+	const void* GetAddress(const int32 Index) const
+	{
+		if (IsValidIndex(Index) == false)
+		{
+			return nullptr;
+		}
+		// Ugly, but FScriptArrayHelper does not give us other option.
+		FPropertyBagArrayRef* NonConstThis = const_cast<FPropertyBagArrayRef*>(this);
+		return (void*)NonConstThis->GetRawPtr(Index);
+	}
+
+	void* GetMutableAddress(const int32 Index) const
+	{
+		if (IsValidIndex(Index) == false)
+		{
+			return nullptr;
+		}
+		// Ugly, but FScriptArrayHelper does not give us other option.
+		FPropertyBagArrayRef* NonConstThis = const_cast<FPropertyBagArrayRef*>(this);
+		return (void*)NonConstThis->GetRawPtr(Index);
+	}
+
+public:	
+	FORCEINLINE FPropertyBagArrayRef(const FPropertyBagPropertyDesc& InDesc, const void* InArray)
+		: FScriptArrayHelper(CastField<FArrayProperty>(InDesc.CachedProperty), InArray)
+	{
+		const FArrayProperty* ArrayProperty = CastField<FArrayProperty>(InDesc.CachedProperty);
+		check(ArrayProperty);
+		check(ArrayProperty->Inner);
+		// Create dummy desc for the inner property. 
+		ValueDesc.ValueType = InDesc.ValueType;
+		ValueDesc.ValueTypeObject = InDesc.ValueTypeObject;
+		ValueDesc.CachedProperty = ArrayProperty->Inner;
+		ValueDesc.ContainerType = EPropertyBagContainerType::None;
+	}
+
+	/**
+	 * Getters
+	 * Numeric types (bool, int32, int64, float, double) support type conversion.
+	 */
+	
+	TValueOrError<bool, EPropertyBagResult> GetValueBool(const int32 Index) const;
+	TValueOrError<uint8, EPropertyBagResult> GetValueByte(const int32 Index) const;
+	TValueOrError<int32, EPropertyBagResult> GetValueInt32(const int32 Index) const;
+	TValueOrError<int64, EPropertyBagResult> GetValueInt64(const int32 Index) const;
+	TValueOrError<float, EPropertyBagResult> GetValueFloat(const int32 Index) const;
+	TValueOrError<double, EPropertyBagResult> GetValueDouble(const int32 Index) const;
+	TValueOrError<FName, EPropertyBagResult> GetValueName(const int32 Index) const;
+	TValueOrError<FString, EPropertyBagResult> GetValueString(const int32 Index) const;
+	TValueOrError<FText, EPropertyBagResult> GetValueText(const int32 Index) const;
+	TValueOrError<uint8, EPropertyBagResult> GetValueEnum(const int32 Index, const UEnum* RequestedEnum) const;
+	TValueOrError<FStructView, EPropertyBagResult> GetValueStruct(const int32 Index, const UScriptStruct* RequestedStruct = nullptr) const;
+	TValueOrError<UObject*, EPropertyBagResult> GetValueObject(const int32 Index, const UClass* RequestedClass = nullptr) const;
+	TValueOrError<UClass*, EPropertyBagResult> GetValueClass(const int32 Index) const;
+
+	/** @return enum value of specified type. */
+	template <typename T>
+	TValueOrError<T, EPropertyBagResult> GetValueEnum(const int32 Index) const
+	{
+		static_assert(TIsEnum<T>::Value, "Should only call this with enum types");
+		
+		TValueOrError<uint8, EPropertyBagResult> Result = GetValueEnum(Index, StaticEnum<T>());
+		if (!Result.IsValid())
+		{
+			return MakeError(Result.GetError());
+		}
+		return MakeValue((T)Result.GetValue());
+	}
+	
+	/** @return struct reference of specified type. */
+	template <typename T>
+	TValueOrError<T*, EPropertyBagResult> GetValueStruct(const int32 Index) const
+	{
+		TValueOrError<FStructView, EPropertyBagResult> Result = GetValueStruct(Index, TBaseStructure<T>::Get());
+		if (!Result.IsValid())
+		{
+			return MakeError(Result.GetError());
+		}
+		if (T* ValuePtr = Result.GetValue().GetMutablePtr<T>())
+		{
+			return MakeValue(ValuePtr);
+		}
+		return MakeError(EPropertyBagResult::TypeMismatch);
+	}
+	
+	/** @return object pointer value of specified type. */
+	template <typename T>
+	TValueOrError<T*, EPropertyBagResult> GetValueObject(const int32 Index) const
+	{
+		static_assert(TIsDerivedFrom<T, UObject>::Value, "Should only call this with object types");
+		
+		TValueOrError<UObject*, EPropertyBagResult> Result = GetValueObject(Index, T::StaticClass());
+		if (!Result.IsValid())
+		{
+			return MakeError(Result.GetError());
+		}
+		if (Result.GetValue() == nullptr)
+		{
+			return MakeValue(nullptr);
+		}
+		if (T* Object = Cast<T>(Result.GetValue()))
+		{
+			return MakeValue(Object);
+		}
+		return MakeError(EPropertyBagResult::TypeMismatch);
+	}
+
+	/**
+	 * Value Setters. A property must exists in that bag before it can be set.  
+	 * Numeric types (bool, int32, int64, float, double) support type conversion.
+	 */
+	EPropertyBagResult SetValueBool(const int32 Index, const bool bInValue) const;
+	EPropertyBagResult SetValueByte(const int32 Index, const uint8 InValue) const;
+	EPropertyBagResult SetValueInt32(const int32 Index, const int32 InValue) const;
+	EPropertyBagResult SetValueInt64(const int32 Index, const int64 InValue) const;
+	EPropertyBagResult SetValueFloat(const int32 Index, const float InValue) const;
+	EPropertyBagResult SetValueDouble(const int32 Index, const double InValue) const;
+	EPropertyBagResult SetValueName(const int32 Index, const FName InValue) const;
+	EPropertyBagResult SetValueString(const int32 Index, const FString& InValue) const;
+	EPropertyBagResult SetValueText(const int32 Index, const FText& InValue) const;
+	EPropertyBagResult SetValueEnum(const int32 Index, const uint8 InValue, const UEnum* Enum) const;
+	EPropertyBagResult SetValueStruct(const int32 Index, FConstStructView InValue) const;
+	EPropertyBagResult SetValueObject(const int32 Index, UObject* InValue) const;
+	EPropertyBagResult SetValueClass(const int32 Index, UClass* InValue) const;
+
+	/** Sets enum value specified type. */
+	template <typename T>
+	EPropertyBagResult SetValueEnum(const int32 Index, const T InValue) const
+	{
+		static_assert(TIsEnum<T>::Value, "Should only call this with enum types");
+		return SetValueEnum(Index, (uint8)InValue, StaticEnum<T>());
+	}
+
+	/** Sets struct value specified type. */
+	template <typename T>
+	EPropertyBagResult SetValueStruct(const int32 Index, const T& InValue) const
+	{
+		return SetValueStruct(Index, FConstStructView::Make(InValue));
+	}
+
+	/** Sets object pointer value specified type. */
+	template <typename T>
+	EPropertyBagResult SetValueObject(const int32 Index, T* InValue) const
+	{
+		static_assert(TIsDerivedFrom<T, UObject>::Value, "Should only call this with object types");
+		return SetValueObject(Index, (UObject*)InValue);
+	}
+};
+
+
 /**
  * Dummy types used to mark up missing types when creating property bags. These are used in the UI to display error message.
  */
diff --git a/Engine/Plugins/Experimental/StructUtils/Source/StructUtilsEditor/Private/Customizations/PropertyBagDetails.cpp b/Engine/Plugins/Experimental/StructUtils/Source/StructUtilsEditor/Private/Customizations/PropertyBagDetails.cpp
--- a/Engine/Plugins/Experimental/StructUtils/Source/StructUtilsEditor/Private/Customizations/PropertyBagDetails.cpp
+++ b/Engine/Plugins/Experimental/StructUtils/Source/StructUtilsEditor/Private/Customizations/PropertyBagDetails.cpp
@@ -199,12 +199,23 @@ FEdGraphPinType GetPropertyDescAsPin(const FPropertyBagPropertyDesc& Desc)
 {
 	UEnum* PropertyTypeEnum = StaticEnum<EPropertyBagPropertyType>();
 	check(PropertyTypeEnum);
-	const UEdGraphSchema_K2* Schema = GetDefault<UEdGraphSchema_K2>();
+	const UPropertyBagSchema* Schema = GetDefault<UPropertyBagSchema>();
 	check(Schema);
 
 	FEdGraphPinType PinType;
 	PinType.PinSubCategory = NAME_None;
 
+	// Container type
+	switch (Desc.ContainerType)
+	{
+	case EPropertyBagContainerType::Array:
+		PinType.ContainerType = EPinContainerType::Array;
+		break;
+	default:
+		PinType.ContainerType = EPinContainerType::None;
+	}
+
+	// Value type
 	switch (Desc.ValueType)
 	{
 	case EPropertyBagPropertyType::Bool:
@@ -272,9 +283,20 @@ FEdGraphPinType GetPropertyDescAsPin(const FPropertyBagPropertyDesc& Desc)
 /** Sets property descriptor based on a Blueprint pin type. */
 void SetPropertyDescFromPin(FPropertyBagPropertyDesc& Desc, const FEdGraphPinType& PinType)
 {
-	const UEdGraphSchema_K2* Schema = GetDefault<UEdGraphSchema_K2>();
+	const UPropertyBagSchema* Schema = GetDefault<UPropertyBagSchema>();
 	check(Schema);
 
+	// Container type
+	switch (PinType.ContainerType)
+	{
+	case EPinContainerType::Array:
+		Desc.ContainerType = EPropertyBagContainerType::Array;
+		break;
+	default:
+		Desc.ContainerType = EPropertyBagContainerType::None;
+	}
+	
+	// Value type
 	if (PinType.PinCategory == UEdGraphSchema_K2::PC_Boolean)
 	{
 		Desc.ValueType = EPropertyBagPropertyType::Bool;
@@ -537,7 +559,7 @@ TSharedRef<SWidget> FPropertyBagInstanceDataDetails::OnPropertyNameContent(TShar
 	auto GetFilteredVariableTypeTree = [](TArray<TSharedPtr<UEdGraphSchema_K2::FPinTypeTreeInfo>>& TypeTree, ETypeTreeFilter TypeTreeFilter)
 	{
 		check(GetDefault<UEdGraphSchema_K2>());
-		GetDefault<UEdGraphSchema_K2>()->GetVariableTypeTree(TypeTree, TypeTreeFilter);
+		GetDefault<UPropertyBagSchema>()->GetVariableTypeTree(TypeTree, TypeTreeFilter);
 
 		// Filter
 		for (TSharedPtr<UEdGraphSchema_K2::FPinTypeTreeInfo>& PinType : TypeTree)
@@ -621,18 +643,23 @@ TSharedRef<SWidget> FPropertyBagInstanceDataDetails::OnPropertyNameContent(TShar
 			}
 		});
 	};
-	
-	MenuBuilder.AddMenuEntry(
-		FUIAction(),
-		SNew(SPinTypeSelector, FGetPinTypeTree::CreateLambda(GetFilteredVariableTypeTree))
-			.TargetPinType_Lambda(GetPinInfo)
-			.OnPinTypeChanged_Lambda(PinInfoChanged)
-			.Schema(GetDefault<UEdGraphSchema_K2>())
-			.bAllowArrays(false)
-			.TypeTreeFilter(ETypeTreeFilter::None)
-			.Font( IDetailLayoutBuilder::GetDetailFont() )
-		);
 
+
+	MenuBuilder.AddWidget(
+		SNew(SBox)
+		.HAlign(HAlign_Right)
+		.Padding(FMargin(12, 0, 12, 0))
+		[
+			SNew(SPinTypeSelector, FGetPinTypeTree::CreateLambda(GetFilteredVariableTypeTree))
+				.TargetPinType_Lambda(GetPinInfo)
+				.OnPinTypeChanged_Lambda(PinInfoChanged)
+				.Schema(GetDefault<UPropertyBagSchema>())
+				.bAllowArrays(true)
+				.TypeTreeFilter(ETypeTreeFilter::None)
+				.Font( IDetailLayoutBuilder::GetDetailFont() )
+		],
+		FText::GetEmpty());
+	
 	MenuBuilder.AddSeparator();
 	
 	MenuBuilder.AddMenuEntry(
@@ -727,7 +754,7 @@ TSharedPtr<SWidget> FPropertyBagDetails::MakeAddPropertyWidget(TSharedPtr<IPrope
 			SNew(SButton)
 			.VAlign(VAlign_Center)
 			.HAlign(HAlign_Center)
-			.ButtonStyle(FAppStyle::Get(), "HoverHintOnly")
+			.ButtonStyle(FAppStyle::Get(), "SimpleButton")
 			.ToolTipText(LOCTEXT("AddProperty_Tooltip", "Add new property"))
 			.OnClicked_Lambda([InStructProperty, InPropUtils]()
 			{
@@ -758,8 +785,19 @@ TSharedPtr<SWidget> FPropertyBagDetails::MakeAddPropertyWidget(TSharedPtr<IPrope
 			[
 				SNew(SImage)
 				.Image(FAppStyle::GetBrush("Icons.PlusCircle"))
+				.ColorAndOpacity(FSlateColor::UseForeground())
 			]
 		];
+
 }
 
+////////////////////////////////////
+
+bool UPropertyBagSchema::SupportsPinTypeContainer(TWeakPtr<const FEdGraphSchemaAction> SchemaAction,
+		const FEdGraphPinType& PinType, const EPinContainerType& ContainerType) const
+{
+	return ContainerType == EPinContainerType::None || ContainerType == EPinContainerType::Array;
+}
+
+
 #undef LOCTEXT_NAMESPACE
diff --git a/Engine/Plugins/Experimental/StructUtils/Source/StructUtilsEditor/Public/PropertyBagDetails.h b/Engine/Plugins/Experimental/StructUtils/Source/StructUtilsEditor/Public/PropertyBagDetails.h
--- a/Engine/Plugins/Experimental/StructUtils/Source/StructUtilsEditor/Public/PropertyBagDetails.h
+++ b/Engine/Plugins/Experimental/StructUtils/Source/StructUtilsEditor/Public/PropertyBagDetails.h
@@ -5,6 +5,7 @@
 #include "IPropertyTypeCustomization.h"
 #include "InstancedStructDetails.h"
 #include "EdGraphSchema_K2.h"
+#include "PropertyBagDetails.generated.h"
 
 class IPropertyHandle;
 class IDetailPropertyRow;
@@ -63,3 +64,14 @@ protected:
 	IPropertyUtilities* PropUtils = nullptr;
 	bool bFixedLayout = false;
 };
+
+/**
+ * Specific property bag schema to allow customizing the requirements (e.g. supported containers).
+ */
+UCLASS()
+class STRUCTUTILSEDITOR_API UPropertyBagSchema : public UEdGraphSchema_K2
+{
+	GENERATED_BODY()
+public:
+	virtual bool SupportsPinTypeContainer(TWeakPtr<const FEdGraphSchemaAction> SchemaAction, const FEdGraphPinType& PinType, const EPinContainerType& ContainerType) const override;
+};
diff --git a/Engine/Plugins/Experimental/StructUtils/Source/StructUtilsTestSuite/Private/PropertyBagTest.cpp b/Engine/Plugins/Experimental/StructUtils/Source/StructUtilsTestSuite/Private/PropertyBagTest.cpp
--- a/Engine/Plugins/Experimental/StructUtils/Source/StructUtilsTestSuite/Private/PropertyBagTest.cpp
+++ b/Engine/Plugins/Experimental/StructUtils/Source/StructUtilsTestSuite/Private/PropertyBagTest.cpp
@@ -20,7 +20,6 @@ struct FTest_CreatePropertyBag : FAITestBase
 		static const FName TemperatureName(TEXT("Temperature"));
 		static const FName CountName(TEXT("Count"));
 
-
 		FInstancedPropertyBag Bag;
 
 		Bag.AddProperty(IsHotName, EPropertyBagPropertyType::Bool);
@@ -93,6 +92,10 @@ struct FTest_MigrateProperty : FAITestBase
 		FInstancedPropertyBag Bag;
 		Bag.AddProperty(TemperatureName, EPropertyBagPropertyType::Float);
 		AITEST_TRUE(TEXT("Bag should have Temperature property"), Bag.FindPropertyDescByName(TemperatureName) != nullptr);
+
+		TValueOrError<float, EPropertyBagResult> FloatDefaultRes = Bag.GetValueFloat(TemperatureName);
+		AITEST_TRUE(TEXT("Bag getting Temperature default value should succeed"), FloatDefaultRes.IsValid());
+		AITEST_TRUE(TEXT("Bag Temperature default value should be 0"), FMath::IsNearlyEqual(FloatDefaultRes.GetValue(), 0.0f));
 		
 		AITEST_TRUE(TEXT("Bag set Temperature as float should succeed"), Bag.SetValueFloat(TemperatureName, 451.0f) == EPropertyBagResult::Success);
 		TValueOrError<float, EPropertyBagResult> FloatRes = Bag.GetValueFloat(TemperatureName);
@@ -293,8 +296,46 @@ struct FTest_GC : FAITestBase
 };
 IMPLEMENT_AI_INSTANT_TEST(FTest_GC, "System.StructUtils.PropertyBag.GC");
 
+struct FTest_Arrays : FAITestBase
+{
+	virtual bool InstantTest() override
+	{
+		static const FName FloatArrayName(TEXT("FloatArray"));
+
+		FInstancedPropertyBag Bag;
+		Bag.AddProperties({
+			{ FloatArrayName, EPropertyBagContainerType::Array, EPropertyBagPropertyType::Float },
+		});
+
+		TValueOrError<FPropertyBagArrayRef, EPropertyBagResult> FloatArrayRes = Bag.GetArrayRef(FloatArrayName);
+		AITEST_TRUE(TEXT("Get float array should succeed"), FloatArrayRes.IsValid());
 
+		FPropertyBagArrayRef FloatArray = FloatArrayRes.GetValue();
+		const int32 FloatIndex = FloatArray.AddValue();
+		AITEST_TRUE(TEXT("Float array should have 1 item"), FloatArray.Num() == 1);
 
+		const TValueOrError<float, EPropertyBagResult> GetDefaultFloatRes = FloatArray.GetValueFloat(FloatIndex);
+		AITEST_TRUE(TEXT("Get float should succeed immediatelly after add"), GetDefaultFloatRes.IsValid());
+		AITEST_TRUE(TEXT("Default value for Float should be 0.0f"), FMath::IsNearlyEqual(GetDefaultFloatRes.GetValue(), 0.0f));
+
+		const EPropertyBagResult SetFloatRes = FloatArray.SetValueFloat(FloatIndex, 123.0f);
+		AITEST_TRUE(TEXT("Set float should succeed"), SetFloatRes == EPropertyBagResult::Success);
+
+		const TValueOrError<float, EPropertyBagResult> GetFloatRes = FloatArray.GetValueFloat(FloatIndex);
+		AITEST_TRUE(TEXT("Get float should succeed"), GetFloatRes.IsValid());
+		AITEST_TRUE(TEXT("Float value should be 123.0f"), FMath::IsNearlyEqual(GetFloatRes.GetValue(), 123.0f));
+
+		const TValueOrError<float, EPropertyBagResult> GetFloatOOBRes = FloatArray.GetValueFloat(42);
+		AITEST_FALSE(TEXT("Get float out of bounds should not succeed"), GetFloatOOBRes.IsValid());
+		AITEST_TRUE(TEXT("Error should be our of bounds"), GetFloatOOBRes.GetError() == EPropertyBagResult::OutOfBounds);
+
+		const EPropertyBagResult SetFloatOOBRes = FloatArray.SetValueFloat(-1, 0.0);
+		AITEST_TRUE(TEXT("Set float out of bounds should return out of bounds"), SetFloatOOBRes == EPropertyBagResult::OutOfBounds);
+		
+		return true;
+	}
+};
+IMPLEMENT_AI_INSTANT_TEST(FTest_Arrays, "System.StructUtils.PropertyBag.Arrays");
 
 
 } // FPropertyBagTest
