Struct utils: Fixed Instanced Struct and ProprtyBag to work on ActorComponents

diff --git a/Engine/Plugins/Experimental/StructUtils/Source/StructUtilsEditor/Private/Customizations/InstancedStructDetails.cpp b/Engine/Plugins/Experimental/StructUtils/Source/StructUtilsEditor/Private/Customizations/InstancedStructDetails.cpp
--- a/Engine/Plugins/Experimental/StructUtils/Source/StructUtilsEditor/Private/Customizations/InstancedStructDetails.cpp
+++ b/Engine/Plugins/Experimental/StructUtils/Source/StructUtilsEditor/Private/Customizations/InstancedStructDetails.cpp
@@ -142,10 +142,24 @@ FInstancedStructDataDetails::FInstancedStructDataDetails(TSharedPtr<IPropertyHan
 
 void FInstancedStructDataDetails::OnStructValuePreChange()
 {
+	PreChangeOuterObjects.Reset();
+
 	// Forward the change event to the real struct handle
 	if (StructProperty && StructProperty->IsValidHandle())
 	{
 		StructProperty->NotifyPreChange();
+
+		// Store the outer objects for the Actor Component special case (see FInstancedStructDataDetails::OnStructValuePostChange).
+		TArray<UObject*> OuterObjects;
+		StructProperty->GetOuterObjects(OuterObjects);
+		
+		for (UObject* Outer : OuterObjects)
+		{
+			if (Outer != nullptr)
+			{
+				PreChangeOuterObjects.Add(FSoftObjectPtr(Outer));
+			}
+		}
 	}
 }
 
@@ -156,6 +170,24 @@ void FInstancedStructDataDetails::OnStructValuePostChange()
 	{
 		TGuardValue<bool> HandlingStructValuePostChangeGuard(bIsHandlingStructValuePostChange, true);
 
+		// When an InstancedStruct is on an Actor Component, the details customization gets rebuild between
+		// the OnStructValuePreChange() and OnStructValuePostChange() calls due to AActor::RerunConstructionScripts().
+		// When the new details panel is build, it will clear the outer objects of the StructProperty property handle, and then setting the value will fail.
+		// To overcome that, we restore the outer objects based on the objects stored in OnStructValuePreChange(). 
+		if (StructProperty->GetNumOuterObjects() == 0 && PreChangeOuterObjects.Num() > 0)
+		{
+			TArray<UObject*> OuterObjects;
+			for (FSoftObjectPtr& SoftOuterObject : PreChangeOuterObjects)
+			{
+				if (UObject* OuterObject = SoftOuterObject.Get())
+				{
+					OuterObjects.Add(OuterObject);
+				}
+			}
+			
+			StructProperty->ReplaceOuterObjects(OuterObjects);
+		}
+		
 		// Copy the modified struct data back to the source instances
 		{
 			FInstancedStruct TmpInstancedStruct;
diff --git a/Engine/Plugins/Experimental/StructUtils/Source/StructUtilsEditor/Private/Customizations/PropertyBagDetails.cpp b/Engine/Plugins/Experimental/StructUtils/Source/StructUtilsEditor/Private/Customizations/PropertyBagDetails.cpp
--- a/Engine/Plugins/Experimental/StructUtils/Source/StructUtilsEditor/Private/Customizations/PropertyBagDetails.cpp
+++ b/Engine/Plugins/Experimental/StructUtils/Source/StructUtilsEditor/Private/Customizations/PropertyBagDetails.cpp
@@ -444,11 +444,11 @@ bool CanHaveMemberVariableOfType(const FEdGraphPinType& PinType)
 
 FPropertyBagInstanceDataDetails::FPropertyBagInstanceDataDetails(TSharedPtr<IPropertyHandle> InValueProperty, TSharedPtr<IPropertyHandle> InStructProperty, IPropertyUtilities* InPropUtils, const bool bInFixedLayout)
 	: FInstancedStructDataDetails(InValueProperty)
-	, StructProperty(InStructProperty)
+	, BagStructProperty(InStructProperty)
 	, PropUtils(InPropUtils)
 	, bFixedLayout(bInFixedLayout)
 {
-	ensure(UE::StructUtils::Private::IsScriptStruct<FInstancedPropertyBag>(StructProperty));
+	ensure(UE::StructUtils::Private::IsScriptStruct<FInstancedPropertyBag>(BagStructProperty));
 	ensure(PropUtils != nullptr);
 }
 
@@ -474,7 +474,7 @@ void FPropertyBagInstanceDataDetails::OnChildRowAdded(IDetailPropertyRow& ChildR
 			.OnVerifyTextChanged_Lambda([this, ChildPropertyHandle](const FText& InText, FText& OutErrorMessage)
 			{
 				const FName NewName = UE::StructUtils::Private::GetValidPropertyName(InText.ToString());
-				bool bResult = UE::StructUtils::Private::IsUniqueName(NewName, ChildPropertyHandle->GetProperty()->GetFName(), StructProperty);
+				bool bResult = UE::StructUtils::Private::IsUniqueName(NewName, ChildPropertyHandle->GetProperty()->GetFName(), BagStructProperty);
 				if (!bResult)
 				{
 					OutErrorMessage = LOCTEXT("MustBeUniqueName", "Property must have unique name");
@@ -489,13 +489,13 @@ void FPropertyBagInstanceDataDetails::OnChildRowAdded(IDetailPropertyRow& ChildR
 				}
 
 				const FName NewName = UE::StructUtils::Private::GetValidPropertyName(InNewText.ToString());
-				if (!UE::StructUtils::Private::IsUniqueName(NewName, ChildPropertyHandle->GetProperty()->GetFName(), StructProperty))
+				if (!UE::StructUtils::Private::IsUniqueName(NewName, ChildPropertyHandle->GetProperty()->GetFName(), BagStructProperty))
 				{
 					return;
 				}
 
 				UE::StructUtils::Private::ApplyChangesToPropertyDescs(
-					LOCTEXT("OnPropertyNameChanged", "Change Property Name"), StructProperty, PropUtils,
+					LOCTEXT("OnPropertyNameChanged", "Change Property Name"), BagStructProperty, PropUtils,
 					[&NewName, &ChildPropertyHandle](TArray<FPropertyBagPropertyDesc>& PropertyDescs)
 					{
 						const FProperty* Property = ChildPropertyHandle->GetProperty();
@@ -587,7 +587,7 @@ TSharedRef<SWidget> FPropertyBagInstanceDataDetails::OnPropertyNameContent(TShar
 
 	auto GetPinInfo = [this, ChildPropertyHandle]()
 	{
-		TArray<FPropertyBagPropertyDesc> PropertyDescs = UE::StructUtils::Private::GetCommonPropertyDescs(StructProperty);
+		TArray<FPropertyBagPropertyDesc> PropertyDescs = UE::StructUtils::Private::GetCommonPropertyDescs(BagStructProperty);
 
 		const FProperty* Property = ChildPropertyHandle->GetProperty();
 		if (FPropertyBagPropertyDesc* Desc = PropertyDescs.FindByPredicate([Property](const FPropertyBagPropertyDesc& Desc){ return Desc.CachedProperty == Property; }))
@@ -601,7 +601,7 @@ TSharedRef<SWidget> FPropertyBagInstanceDataDetails::OnPropertyNameContent(TShar
 	auto PinInfoChanged = [this, ChildPropertyHandle](const FEdGraphPinType& PinType)
 	{
 		UE::StructUtils::Private::ApplyChangesToPropertyDescs(
-			LOCTEXT("OnPropertyTypeChanged", "Change Property Type"), StructProperty, PropUtils,
+			LOCTEXT("OnPropertyTypeChanged", "Change Property Type"), BagStructProperty, PropUtils,
 			[&PinType, &ChildPropertyHandle](TArray<FPropertyBagPropertyDesc>& PropertyDescs)
 			{
 				// Find and change struct type
@@ -616,7 +616,7 @@ TSharedRef<SWidget> FPropertyBagInstanceDataDetails::OnPropertyNameContent(TShar
 	auto RemoveProperty = [this, ChildPropertyHandle]()
 	{
 		UE::StructUtils::Private::ApplyChangesToPropertyDescs(
-			LOCTEXT("OnPropertyRemoved", "Remove Property"), StructProperty, PropUtils,
+			LOCTEXT("OnPropertyRemoved", "Remove Property"), BagStructProperty, PropUtils,
 			[&ChildPropertyHandle](TArray<FPropertyBagPropertyDesc>& PropertyDescs)
 			{
 				const FProperty* Property = ChildPropertyHandle->GetProperty();
@@ -627,7 +627,7 @@ TSharedRef<SWidget> FPropertyBagInstanceDataDetails::OnPropertyNameContent(TShar
 	auto MoveProperty = [this, ChildPropertyHandle](const int32 Delta)
 	{
 		UE::StructUtils::Private::ApplyChangesToPropertyDescs(
-		LOCTEXT("OnPropertyMoved", "Move Property"), StructProperty, PropUtils,
+		LOCTEXT("OnPropertyMoved", "Move Property"), BagStructProperty, PropUtils,
 		[&ChildPropertyHandle, &Delta](TArray<FPropertyBagPropertyDesc>& PropertyDescs)
 		{
 			// Move
diff --git a/Engine/Plugins/Experimental/StructUtils/Source/StructUtilsEditor/Public/InstancedStructDetails.h b/Engine/Plugins/Experimental/StructUtils/Source/StructUtilsEditor/Public/InstancedStructDetails.h
--- a/Engine/Plugins/Experimental/StructUtils/Source/StructUtilsEditor/Public/InstancedStructDetails.h
+++ b/Engine/Plugins/Experimental/StructUtils/Source/StructUtilsEditor/Public/InstancedStructDetails.h
@@ -78,6 +78,9 @@ private:
 	/** Sync the current state of the editable struct instance from the source instance(s) */
 	void SyncEditableInstanceFromSource(bool* OutStructMismatch = nullptr);
 
+	/** Outer objects at the time of PreChangeProperties */
+	TArray<FSoftObjectPtr> PreChangeOuterObjects;
+	
 	/** Handle to the struct property being edited */
 	TSharedPtr<IPropertyHandle> StructProperty;
 
diff --git a/Engine/Plugins/Experimental/StructUtils/Source/StructUtilsEditor/Public/PropertyBagDetails.h b/Engine/Plugins/Experimental/StructUtils/Source/StructUtilsEditor/Public/PropertyBagDetails.h
--- a/Engine/Plugins/Experimental/StructUtils/Source/StructUtilsEditor/Public/PropertyBagDetails.h
+++ b/Engine/Plugins/Experimental/StructUtils/Source/StructUtilsEditor/Public/PropertyBagDetails.h
@@ -60,7 +60,7 @@ public:
 protected:
 	TSharedRef<SWidget> OnPropertyNameContent(TSharedPtr<IPropertyHandle> ChildPropertyHandle, TSharedPtr<SInlineEditableTextBlock> InlineWidget) const;
 
-	TSharedPtr<IPropertyHandle> StructProperty;
+	TSharedPtr<IPropertyHandle> BagStructProperty;
 	IPropertyUtilities* PropUtils = nullptr;
 	bool bFixedLayout = false;
 };
