PR #9282: Blueprint Graph Support for FInstancedStruct (Contributed by [REDACTED_PERSON])

diff --git a/Engine/Plugins/Experimental/StructUtils/Source/StructUtils/Private/StructUtilsFunctionLibrary.cpp b/Engine/Plugins/Experimental/StructUtils/Source/StructUtils/Private/StructUtilsFunctionLibrary.cpp
new file mode 100644
--- /dev/null
+++ b/Engine/Plugins/Experimental/StructUtils/Source/StructUtils/Private/StructUtilsFunctionLibrary.cpp
@@ -0,0 +1,136 @@
+// Copyright Epic Games, Inc. All Rights Reserved.
+
+#include "StructUtilsFunctionLibrary.h"
+
+#define LOCTEXT_NAMESPACE "UStructUtilsFunctionLibrary"
+
+
+FInstancedStruct UStructUtilsFunctionLibrary::MakeInstancedStruct(const int32& Value)
+{
+	// We should never hit this! stubs to avoid NoExport on the class.
+	checkNoEntry();
+	return {};
+}
+
+void UStructUtilsFunctionLibrary::SetInstancedStructValue(FInstancedStruct& InstancedStruct, const int32& Value)
+{
+	// We should never hit this! stubs to avoid NoExport on the class.
+	checkNoEntry();
+}
+
+void UStructUtilsFunctionLibrary::GetInstancedStructValue(EStructUtilsResult& ExecResult, const FInstancedStruct& InstancedStruct, int32& Value)
+{
+	// We should never hit this! stubs to avoid NoExport on the class.
+	checkNoEntry();
+}
+
+DEFINE_FUNCTION(UStructUtilsFunctionLibrary::execMakeInstancedStruct)
+{
+	Stack.Step(Stack.Object, nullptr);
+
+	Stack.MostRecentPropertyAddress = nullptr;
+	Stack.MostRecentPropertyContainer = nullptr;
+	Stack.StepCompiledIn<FStructProperty>(nullptr);
+	
+	const FStructProperty* ValueProp = CastField<FStructProperty>(Stack.MostRecentProperty);
+	const void* ValuePtr = Stack.MostRecentPropertyAddress;
+
+	P_FINISH;
+
+	if (!ValueProp || !ValuePtr)
+	{
+		FBlueprintExceptionInfo ExceptionInfo(
+			EBlueprintExceptionType::AbortExecution,
+			LOCTEXT("InstancedStruct_MakeInvalidValueWarning", "Failed to resolve the Value for Make Instanced Struct")
+		);
+
+		FBlueprintCoreDelegates::ThrowScriptException(P_THIS, Stack, ExceptionInfo);
+
+		P_NATIVE_BEGIN;
+		(*(FInstancedStruct*)RESULT_PARAM).Reset();
+		P_NATIVE_END;
+	}
+	else
+	{
+		P_NATIVE_BEGIN;
+		(*(FInstancedStruct*)RESULT_PARAM).InitializeAs(ValueProp->Struct, (const uint8*)ValuePtr);
+		P_NATIVE_END;
+	}
+}
+
+DEFINE_FUNCTION(UStructUtilsFunctionLibrary::execSetInstancedStructValue)
+{
+	P_GET_STRUCT_REF(FInstancedStruct, InstancedStruct);
+
+	Stack.MostRecentPropertyAddress = nullptr;
+	Stack.MostRecentPropertyContainer = nullptr;
+	Stack.StepCompiledIn<FStructProperty>(nullptr);
+	
+	const FStructProperty* ValueProp = CastField<FStructProperty>(Stack.MostRecentProperty);
+	const void* ValuePtr = Stack.MostRecentPropertyAddress;
+
+	P_FINISH;
+
+	if (!ValueProp || !ValuePtr)
+	{
+		FBlueprintExceptionInfo ExceptionInfo(
+			EBlueprintExceptionType::AbortExecution,
+			LOCTEXT("InstancedStruct_SetInvalidValueWarning", "Failed to resolve the Value for Set Instanced Struct Value")
+		);
+
+		FBlueprintCoreDelegates::ThrowScriptException(P_THIS, Stack, ExceptionInfo);
+
+		P_NATIVE_BEGIN;
+		InstancedStruct.Reset();
+		P_NATIVE_END;
+	}
+	else
+	{
+		P_NATIVE_BEGIN;
+		InstancedStruct.InitializeAs(ValueProp->Struct, (const uint8*)ValuePtr);
+		P_NATIVE_END;
+	}
+}
+
+DEFINE_FUNCTION(UStructUtilsFunctionLibrary::execGetInstancedStructValue)
+{
+	P_GET_ENUM_REF(EStructUtilsResult, ExecResult);
+	P_GET_STRUCT_REF(FInstancedStruct, InstancedStruct);
+
+	Stack.MostRecentPropertyAddress = nullptr;
+	Stack.MostRecentPropertyContainer = nullptr;
+	Stack.StepCompiledIn<FStructProperty>(nullptr);
+	
+	const FStructProperty* ValueProp = CastField<FStructProperty>(Stack.MostRecentProperty);
+	void* ValuePtr = Stack.MostRecentPropertyAddress;
+
+	P_FINISH;
+
+	ExecResult = EStructUtilsResult::NotValid;
+
+	if (!ValueProp || !ValuePtr)
+	{
+		FBlueprintExceptionInfo ExceptionInfo(
+			EBlueprintExceptionType::AbortExecution,
+			LOCTEXT("InstancedStruct_GetInvalidValueWarning", "Failed to resolve the Value for Get Instanced Struct Value")
+		);
+
+		FBlueprintCoreDelegates::ThrowScriptException(P_THIS, Stack, ExceptionInfo);
+	}
+	else
+	{
+		P_NATIVE_BEGIN;
+		if (InstancedStruct.IsValid() && InstancedStruct.GetScriptStruct()->IsChildOf(ValueProp->Struct))
+		{
+			ValueProp->Struct->CopyScriptStruct(ValuePtr, InstancedStruct.GetMemory());
+			ExecResult = EStructUtilsResult::Valid;
+		}
+		else
+		{
+			ExecResult = EStructUtilsResult::NotValid;
+		}
+		P_NATIVE_END;
+	}
+}
+
+#undef LOCTEXT_NAMESPACE
\ No newline at end of file
diff --git a/Engine/Plugins/Experimental/StructUtils/Source/StructUtils/Public/InstancedStruct.h b/Engine/Plugins/Experimental/StructUtils/Source/StructUtils/Public/InstancedStruct.h
--- a/Engine/Plugins/Experimental/StructUtils/Source/StructUtils/Public/InstancedStruct.h
+++ b/Engine/Plugins/Experimental/StructUtils/Source/StructUtils/Public/InstancedStruct.h
@@ -18,7 +18,7 @@ struct FConstSharedStruct;
  *	UPROPERTY(EditAnywhere, Category = Foo, meta = (BaseStruct = "/Script/ModuleName.TestStructBase"))
  *	TArray<FInstancedStruct> TestArray;
  */
-USTRUCT()
+USTRUCT(BlueprintType)
 struct STRUCTUTILS_API FInstancedStruct
 {
 	GENERATED_BODY()
diff --git a/Engine/Plugins/Experimental/StructUtils/Source/StructUtils/Public/StructUtilsFunctionLibrary.h b/Engine/Plugins/Experimental/StructUtils/Source/StructUtils/Public/StructUtilsFunctionLibrary.h
new file mode 100644
--- /dev/null
+++ b/Engine/Plugins/Experimental/StructUtils/Source/StructUtils/Public/StructUtilsFunctionLibrary.h
@@ -0,0 +1,83 @@
+// Copyright Epic Games, Inc. All Rights Reserved.
+
+#pragma once
+
+#include "Kismet/BlueprintFunctionLibrary.h"
+#include "InstancedStruct.h"
+#include "StructUtilsFunctionLibrary.generated.h"
+
+struct FGenericStruct;
+
+UENUM()
+enum class EStructUtilsResult : uint8
+{
+	Valid,
+	NotValid,
+};
+
+UCLASS()
+class STRUCTUTILS_API UStructUtilsFunctionLibrary : public UBlueprintFunctionLibrary
+{
+	GENERATED_BODY()
+
+public:
+
+	/**
+	 * Create a new InstancedStruct from the given source value.
+	 */
+	UFUNCTION(BlueprintCallable, CustomThunk, Category = "Utilities|Instanced Struct", meta = (CustomStructureParam = "Value", BlueprintInternalUseOnly="true"))
+	static FInstancedStruct MakeInstancedStruct(const int32& Value);
+
+	/**
+	 * Sets the value of InstancedStruct from the given source value.
+	 */
+	UFUNCTION(BlueprintCallable, CustomThunk, Category = "Utilities|Instanced Struct", meta = (CustomStructureParam = "Value", BlueprintInternalUseOnly="true"))
+	static void SetInstancedStructValue(UPARAM(Ref) FInstancedStruct& InstancedStruct, const int32& Value);
+
+	/**
+	 * Retrieves data from an InstancedStruct if it matches the output type.
+	 */
+	UFUNCTION(BlueprintCallable, CustomThunk, Category = "Utilities|Instanced Struct", meta = (CustomStructureParam = "Value", ExpandEnumAsExecs = "ExecResult", BlueprintInternalUseOnly="true"))
+	static void GetInstancedStructValue(EStructUtilsResult& ExecResult, UPARAM(Ref) const FInstancedStruct& InstancedStruct, int32& Value);
+
+	/**
+	 * Resets an InstancedStruct.
+	 */
+	UFUNCTION(BlueprintCallable, Category = "Utilities|Instanced Struct", meta = (AdvancedDisplay = "1"))
+	static void Reset(UPARAM(Ref)FInstancedStruct& InstancedStruct, const UScriptStruct* StructType = nullptr)
+	{
+		InstancedStruct.InitializeAs(StructType, nullptr);
+	}
+
+	/** 
+	 * Checks whether an InstancedStruct contains value.
+	 */
+	UFUNCTION(BlueprintCallable, Category = "Utilities|Instanced Struct", meta = (DisplayName = "Is Valid (Branch)", ExpandEnumAsExecs = "ReturnValue"))
+	static EStructUtilsResult IsInstancedStructValid(UPARAM(Ref) const FInstancedStruct& InstancedStruct)
+	{
+		return InstancedStruct.IsValid() ? EStructUtilsResult::Valid : EStructUtilsResult::NotValid;
+	}
+
+	/**
+	 * Checks whether two InstancedStructs (and the values contained within) are equal.
+	 */
+	UFUNCTION(BlueprintPure, Category = "Utilities|Instanced Struct", meta = (CompactNodeTitle = "==", DisplayName = "Equal", Keywords = "== equal"))
+	static bool EqualEqual_InstancedStruct(UPARAM(Ref) const FInstancedStruct& A, UPARAM(Ref) const FInstancedStruct& B) { return A == B; }
+
+	/**
+	 * Checks whether two InstancedStructs are not equal.
+	 */
+	UFUNCTION(BlueprintPure, Category = "Utilities|Instanced Struct", meta = (CompactNodeTitle = "!=", DisplayName = "Not Equal", Keywords = "!= not equal"))
+	static bool NotEqual_InstancedStruct(UPARAM(Ref) const FInstancedStruct& A, UPARAM(Ref) const FInstancedStruct& B) { return A != B; }
+
+	/**
+	 * Checks whether the InstancedStruct contains value.
+	 */
+	UFUNCTION(BlueprintPure, Category = "Utilities|Instanced Struct", meta = ( DisplayName = "Is Valid"))
+	static bool IsValid_InstancedStruct(UPARAM(Ref) const FInstancedStruct& InstancedStruct) { return InstancedStruct.IsValid(); }
+
+private:
+	DECLARE_FUNCTION(execMakeInstancedStruct);
+	DECLARE_FUNCTION(execSetInstancedStructValue);
+	DECLARE_FUNCTION(execGetInstancedStructValue);
+};
\ No newline at end of file
diff --git a/Engine/Plugins/Experimental/StructUtils/Source/StructUtilsNodes/Private/K2Node_InstancedStruct.cpp b/Engine/Plugins/Experimental/StructUtils/Source/StructUtilsNodes/Private/K2Node_InstancedStruct.cpp
new file mode 100644
--- /dev/null
+++ b/Engine/Plugins/Experimental/StructUtils/Source/StructUtilsNodes/Private/K2Node_InstancedStruct.cpp
@@ -0,0 +1,61 @@
+﻿// Copyright Epic Games, Inc. All Rights Reserved.
+
+#include "K2Node_InstancedStruct.h"
+#include "EdGraphSchema_K2.h"
+#include "BlueprintNodeSpawner.h"
+#include "BlueprintActionDatabaseRegistrar.h"
+#include "StructUtilsFunctionLibrary.h"
+
+#define LOCTEXT_NAMESPACE "InstancedStruct"
+
+void UK2Node_InstancedStruct::GetMenuActions(FBlueprintActionDatabaseRegistrar& ActionRegistrar) const
+{
+	Super::GetMenuActions(ActionRegistrar);
+	UClass* Action = GetClass();
+	if (ActionRegistrar.IsOpenForRegistration(Action))
+	{
+		auto CustomizeLambda = [](UEdGraphNode* NewNode, bool bIsTemplateNode, const FName FunctionName)
+		{
+			UK2Node_InstancedStruct* Node = CastChecked<UK2Node_InstancedStruct>(NewNode);
+			UFunction* Function = UStructUtilsFunctionLibrary::StaticClass()->FindFunctionByName(FunctionName);
+			check(Function);
+			Node->SetFromFunction(Function);
+		};
+		
+		// MakeInstancedStruct()
+		UBlueprintNodeSpawner* MakeNodeSpawner = UBlueprintNodeSpawner::Create(GetClass());
+		check(MakeNodeSpawner != nullptr);
+		MakeNodeSpawner->CustomizeNodeDelegate = UBlueprintNodeSpawner::FCustomizeNodeDelegate::CreateStatic(CustomizeLambda, GET_FUNCTION_NAME_CHECKED(UStructUtilsFunctionLibrary, MakeInstancedStruct));
+		ActionRegistrar.AddBlueprintAction(Action, MakeNodeSpawner);
+
+		// SetInstancedStructValue()
+		UBlueprintNodeSpawner* SetNodeSpawner = UBlueprintNodeSpawner::Create(GetClass());
+		check(SetNodeSpawner != nullptr);
+		SetNodeSpawner->CustomizeNodeDelegate = UBlueprintNodeSpawner::FCustomizeNodeDelegate::CreateStatic(CustomizeLambda, GET_FUNCTION_NAME_CHECKED(UStructUtilsFunctionLibrary, SetInstancedStructValue));
+		ActionRegistrar.AddBlueprintAction(Action, SetNodeSpawner);
+
+		// GetInstancedStructValue()
+		UBlueprintNodeSpawner* GetNodeSpawner = UBlueprintNodeSpawner::Create(GetClass());
+		check(GetNodeSpawner != nullptr);
+		GetNodeSpawner->CustomizeNodeDelegate = UBlueprintNodeSpawner::FCustomizeNodeDelegate::CreateStatic(CustomizeLambda, GET_FUNCTION_NAME_CHECKED(UStructUtilsFunctionLibrary, GetInstancedStructValue));
+		ActionRegistrar.AddBlueprintAction(Action, GetNodeSpawner);
+	}
+}
+
+bool UK2Node_InstancedStruct::IsConnectionDisallowed(const UEdGraphPin* MyPin, const UEdGraphPin* OtherPin, FString& OutReason) const
+{
+	const UEdGraphPin* ValuePin = FindPinChecked(FName(TEXT("Value")));
+
+	if (MyPin == ValuePin && MyPin->PinType.PinCategory == UEdGraphSchema_K2::PC_Wildcard)
+	{
+		if (OtherPin->PinType.PinCategory != UEdGraphSchema_K2::PC_Struct)
+		{
+			OutReason = TEXT("Value must be a struct.");
+			return true;
+		}
+	}
+
+	return false;
+}
+
+#undef LOCTEXT_NAMESPACE
diff --git a/Engine/Plugins/Experimental/StructUtils/Source/StructUtilsNodes/Private/StructUtilsNodesModule.cpp b/Engine/Plugins/Experimental/StructUtils/Source/StructUtilsNodes/Private/StructUtilsNodesModule.cpp
new file mode 100644
--- /dev/null
+++ b/Engine/Plugins/Experimental/StructUtils/Source/StructUtilsNodes/Private/StructUtilsNodesModule.cpp
@@ -0,0 +1,5 @@
+﻿// Copyright Epic Games, Inc. All Rights Reserved.
+
+#include "Modules/ModuleManager.h"
+
+IMPLEMENT_MODULE(FDefaultModuleImpl, StructUtilsNodes);
diff --git a/Engine/Plugins/Experimental/StructUtils/Source/StructUtilsNodes/Public/K2Node_InstancedStruct.h b/Engine/Plugins/Experimental/StructUtils/Source/StructUtilsNodes/Public/K2Node_InstancedStruct.h
new file mode 100644
--- /dev/null
+++ b/Engine/Plugins/Experimental/StructUtils/Source/StructUtilsNodes/Public/K2Node_InstancedStruct.h
@@ -0,0 +1,28 @@
+﻿// Copyright Epic Games, Inc. All Rights Reserved.
+
+
+#pragma once
+
+#include "K2Node_CallFunction.h"
+#include "K2Node_InstancedStruct.generated.h"
+
+class UEdGraphPin;
+class FBlueprintActionDatabaseRegistrar;
+
+/**
+ * Node customization for MakeInstancedStruct(), SetInstancedStructValue(), and GetInstancedStructValue().
+ */
+UCLASS()
+class UK2Node_InstancedStruct : public UK2Node_CallFunction
+{
+	GENERATED_BODY()
+
+	//~ Begin UEdGraphNode Interface
+	virtual void GetMenuActions(FBlueprintActionDatabaseRegistrar& ActionRegistrar) const override;
+	//~ End UEdGraphNode Interface
+
+	//~ Begin K2Node Interface
+	virtual bool IsConnectionDisallowed(const UEdGraphPin* MyPin, const UEdGraphPin* OtherPin, FString& OutReason) const override;
+	//~ End K2Node Interface
+};
+
diff --git a/Engine/Plugins/Experimental/StructUtils/Source/StructUtilsNodes/StructUtilsNodes.Build.cs b/Engine/Plugins/Experimental/StructUtils/Source/StructUtilsNodes/StructUtilsNodes.Build.cs
new file mode 100644
--- /dev/null
+++ b/Engine/Plugins/Experimental/StructUtils/Source/StructUtilsNodes/StructUtilsNodes.Build.cs
@@ -0,0 +1,38 @@
+﻿// Copyright Epic Games, Inc. All Rights Reserved.
+
+namespace UnrealBuildTool.Rules
+{
+	public class StructUtilsNodes : ModuleRules
+	{
+		public StructUtilsNodes(ReadOnlyTargetRules Target) : base(Target)
+		{
+			PCHUsage = PCHUsageMode.UseExplicitOrSharedPCHs;
+
+			PrivateDependencyModuleNames.AddRange(
+				new string[]
+				{
+					"BlueprintGraph",
+					"Core",
+					"CoreUObject",
+					"Engine",
+					"KismetCompiler",
+					"PropertyEditor",
+					"StructUtils",
+					"UnrealEd"
+				}
+			);
+
+			PrivateIncludePaths.AddRange(
+				new string[]
+				{
+				}
+			);
+
+			PublicIncludePaths.AddRange(
+				new string[]
+				{
+				}
+			);
+		}
+	}
+}
\ No newline at end of file
diff --git a/Engine/Plugins/Experimental/StructUtils/StructUtils.uplugin b/Engine/Plugins/Experimental/StructUtils/StructUtils.uplugin
--- a/Engine/Plugins/Experimental/StructUtils/StructUtils.uplugin
+++ b/Engine/Plugins/Experimental/StructUtils/StructUtils.uplugin
@@ -25,6 +25,10 @@
 			"Name" : "StructUtilsEditor",
 			"Type" : "Editor"
 		},
+		{
+			"Name": "StructUtilsNodes",
+			"Type": "UncookedOnly"
+		},
 		{
 			"Name": "StructUtilsTestSuite",
 			"Type": "UncookedOnly"
