Instanced Struct: Store PreChangeOuterObjects as paths instead of TSoftObjectPtr<> so that they survive RerunConstructionScripts on blueprinted class.

diff --git a/Engine/Plugins/Experimental/StructUtils/Source/StructUtilsEditor/Private/Customizations/InstancedStructDetails.cpp b/Engine/Plugins/Experimental/StructUtils/Source/StructUtilsEditor/Private/Customizations/InstancedStructDetails.cpp
--- a/Engine/Plugins/Experimental/StructUtils/Source/StructUtilsEditor/Private/Customizations/InstancedStructDetails.cpp
+++ b/Engine/Plugins/Experimental/StructUtils/Source/StructUtilsEditor/Private/Customizations/InstancedStructDetails.cpp
@@ -142,7 +142,7 @@ FInstancedStructDataDetails::FInstancedStructDataDetails(TSharedPtr<IPropertyHan
 
 void FInstancedStructDataDetails::OnStructValuePreChange()
 {
-	PreChangeOuterObjects.Reset();
+	PreChangeOuterObjectNames.Reset();
 
 	// Forward the change event to the real struct handle
 	if (StructProperty && StructProperty->IsValidHandle())
@@ -157,7 +157,7 @@ void FInstancedStructDataDetails::OnStructValuePreChange()
 		{
 			if (Outer != nullptr)
 			{
-				PreChangeOuterObjects.Add(FSoftObjectPtr(Outer));
+				PreChangeOuterObjectNames.Add(Outer->GetPathName());
 			}
 		}
 	}
@@ -173,13 +173,14 @@ void FInstancedStructDataDetails::OnStructValuePostChange()
 		// When an InstancedStruct is on an Actor Component, the details customization gets rebuild between
 		// the OnStructValuePreChange() and OnStructValuePostChange() calls due to AActor::RerunConstructionScripts().
 		// When the new details panel is build, it will clear the outer objects of the StructProperty property handle, and then setting the value will fail.
-		// To overcome that, we restore the outer objects based on the objects stored in OnStructValuePreChange(). 
-		if (StructProperty->GetNumOuterObjects() == 0 && PreChangeOuterObjects.Num() > 0)
+		// To overcome that, we restore the outer objects based on the objects stored in OnStructValuePreChange().
+		if (StructProperty->GetNumOuterObjects() == 0 && PreChangeOuterObjectNames.Num() > 0)
 		{
 			TArray<UObject*> OuterObjects;
-			for (FSoftObjectPtr& SoftOuterObject : PreChangeOuterObjects)
+			for (const FString& ObjectPathName : PreChangeOuterObjectNames)
 			{
-				if (UObject* OuterObject = SoftOuterObject.Get())
+				UObject* OuterObject = FindObject<UObject>(nullptr, *ObjectPathName);
+				if (OuterObject)
 				{
 					OuterObjects.Add(OuterObject);
 				}
diff --git a/Engine/Plugins/Experimental/StructUtils/Source/StructUtilsEditor/Public/InstancedStructDetails.h b/Engine/Plugins/Experimental/StructUtils/Source/StructUtilsEditor/Public/InstancedStructDetails.h
--- a/Engine/Plugins/Experimental/StructUtils/Source/StructUtilsEditor/Public/InstancedStructDetails.h
+++ b/Engine/Plugins/Experimental/StructUtils/Source/StructUtilsEditor/Public/InstancedStructDetails.h
@@ -79,7 +79,7 @@ private:
 	void SyncEditableInstanceFromSource(bool* OutStructMismatch = nullptr);
 
 	/** Outer objects at the time of PreChangeProperties */
-	TArray<FSoftObjectPtr> PreChangeOuterObjects;
+	TArray<FString> PreChangeOuterObjectNames;
 	
 	/** Handle to the struct property being edited */
 	TSharedPtr<IPropertyHandle> StructProperty;
