PropertyBag: UI tweaks
- simplified the usage of the FPropertyBagInstanceDataDetails
- Changed the UI to match closer the array property style

diff --git a/Engine/Plugins/Experimental/StructUtils/Source/StructUtilsEditor/Private/Customizations/PropertyBagDetails.cpp b/Engine/Plugins/Experimental/StructUtils/Source/StructUtilsEditor/Private/Customizations/PropertyBagDetails.cpp
--- a/Engine/Plugins/Experimental/StructUtils/Source/StructUtilsEditor/Private/Customizations/PropertyBagDetails.cpp
+++ b/Engine/Plugins/Experimental/StructUtils/Source/StructUtilsEditor/Private/Customizations/PropertyBagDetails.cpp
@@ -438,12 +438,11 @@ bool CanHaveMemberVariableOfType(const FEdGraphPinType& PinType)
 //----------------------------------------------------------------//
 //  FPropertyBagInstanceDataDetails
 //  - StructProperty is FInstancedPropertyBag
-//  - ValueProperty is FInstancedStruct
-//  - ChildPropertyHandle a child property if the ValueProperty  
+//  - ChildPropertyHandle a child property of the FInstancedPropertyBag::Value (FInstancedStruct)  
 //----------------------------------------------------------------//
 
-FPropertyBagInstanceDataDetails::FPropertyBagInstanceDataDetails(TSharedPtr<IPropertyHandle> InValueProperty, TSharedPtr<IPropertyHandle> InStructProperty, IPropertyUtilities* InPropUtils, const bool bInFixedLayout)
-	: FInstancedStructDataDetails(InValueProperty)
+FPropertyBagInstanceDataDetails::FPropertyBagInstanceDataDetails(TSharedPtr<IPropertyHandle> InStructProperty, IPropertyUtilities* InPropUtils, const bool bInFixedLayout)
+	: FInstancedStructDataDetails(InStructProperty.IsValid() ? InStructProperty->GetChildHandle(TEXT("Value")) : nullptr)
 	, BagStructProperty(InStructProperty)
 	, PropUtils(InPropUtils)
 	, bFixedLayout(bInFixedLayout)
@@ -710,9 +709,6 @@ void FPropertyBagDetails::CustomizeHeader(TSharedRef<class IPropertyHandle> Stru
 	StructProperty = StructPropertyHandle;
 	check(StructProperty);
 	
-	ValueProperty = StructProperty->GetChildHandle(TEXT("Value"));
-	check(ValueProperty);
-
 	static const FName NAME_FixedLayout = "FixedLayout";
 	if (const FProperty* MetaDataProperty = StructProperty->GetMetaDataProperty())
 	{
@@ -741,7 +737,7 @@ void FPropertyBagDetails::CustomizeHeader(TSharedRef<class IPropertyHandle> Stru
 void FPropertyBagDetails::CustomizeChildren(TSharedRef<class IPropertyHandle> StructPropertyHandle, class IDetailChildrenBuilder& StructBuilder, IPropertyTypeCustomizationUtils& StructCustomizationUtils)
 {
 	// Show the Value (FInstancedStruct) as child rows.
-	TSharedRef<FPropertyBagInstanceDataDetails> InstanceDetails = MakeShareable(new FPropertyBagInstanceDataDetails(ValueProperty, StructProperty, PropUtils, bFixedLayout));
+	TSharedRef<FPropertyBagInstanceDataDetails> InstanceDetails = MakeShareable(new FPropertyBagInstanceDataDetails(StructProperty, PropUtils, bFixedLayout));
 	StructBuilder.AddCustomBuilder(InstanceDetails);
 }
 
diff --git a/Engine/Plugins/Experimental/StructUtils/Source/StructUtilsEditor/Public/PropertyBagDetails.h b/Engine/Plugins/Experimental/StructUtils/Source/StructUtilsEditor/Public/PropertyBagDetails.h
--- a/Engine/Plugins/Experimental/StructUtils/Source/StructUtilsEditor/Public/PropertyBagDetails.h
+++ b/Engine/Plugins/Experimental/StructUtils/Source/StructUtilsEditor/Public/PropertyBagDetails.h
@@ -35,7 +35,6 @@ private:
 	
 	/** Handle to the struct property being edited */
 	TSharedPtr<IPropertyHandle> StructProperty;
-	TSharedPtr<IPropertyHandle> ValueProperty;
 
 	bool bFixedLayout = false; 
 	
@@ -53,7 +52,7 @@ private:
 class STRUCTUTILSEDITOR_API FPropertyBagInstanceDataDetails : public FInstancedStructDataDetails
 {
 public:
-	FPropertyBagInstanceDataDetails(TSharedPtr<IPropertyHandle> InValueProperty, TSharedPtr<IPropertyHandle> InStructProperty, IPropertyUtilities* InPropUtils, const bool bInFixedLayout);
+	FPropertyBagInstanceDataDetails(TSharedPtr<IPropertyHandle> InStructProperty, IPropertyUtilities* InPropUtils, const bool bInFixedLayout);
 
 	virtual void OnChildRowAdded(IDetailPropertyRow& ChildRow) override;
 
diff --git a/Engine/Plugins/Runtime/StateTree/Source/StateTreeEditorModule/Private/Customizations/StateTreeStateParametersDetails.cpp b/Engine/Plugins/Runtime/StateTree/Source/StateTreeEditorModule/Private/Customizations/StateTreeStateParametersDetails.cpp
--- a/Engine/Plugins/Runtime/StateTree/Source/StateTreeEditorModule/Private/Customizations/StateTreeStateParametersDetails.cpp
+++ b/Engine/Plugins/Runtime/StateTree/Source/StateTreeEditorModule/Private/Customizations/StateTreeStateParametersDetails.cpp
@@ -23,8 +23,8 @@
 class FStateTreeStateParametersInstanceDataDetails : public FPropertyBagInstanceDataDetails
 {
 public:
-	FStateTreeStateParametersInstanceDataDetails(TSharedPtr<IPropertyHandle> InValueProperty, TSharedPtr<IPropertyHandle> InStructProperty, IPropertyUtilities* InPropUtils, const bool bInFixedLayout, FGuid InID, UStateTreeEditorData* InEditorData)
-		: FPropertyBagInstanceDataDetails(InValueProperty, InStructProperty, InPropUtils, bInFixedLayout)
+	FStateTreeStateParametersInstanceDataDetails(TSharedPtr<IPropertyHandle> InStructProperty, IPropertyUtilities* InPropUtils, const bool bInFixedLayout, FGuid InID, UStateTreeEditorData* InEditorData)
+		: FPropertyBagInstanceDataDetails(InStructProperty, InPropUtils, bInFixedLayout)
 		, EditorData(InEditorData)
 	{
 		EditorPropBindings = EditorData ? EditorData->GetPropertyEditorBindings() : nullptr;
@@ -103,11 +103,9 @@ void FStateTreeStateParametersDetails::CustomizeHeader(TSharedRef<class IPropert
 	ParametersProperty = StructProperty->GetChildHandle(TEXT("Parameters"));
 	FixedLayoutProperty = StructProperty->GetChildHandle(TEXT("bFixedLayout"));
 	IDProperty = StructProperty->GetChildHandle(TEXT("ID"));
-	ParametersValueProperty = ParametersProperty->GetChildHandle(TEXT("Value"));
 	check(ParametersProperty.IsValid());
 	check(FixedLayoutProperty.IsValid());
 	check(IDProperty.IsValid());
-	check(ParametersValueProperty.IsValid());
 
 	FindOuterObjects();
 	
@@ -138,7 +136,7 @@ void FStateTreeStateParametersDetails::CustomizeChildren(TSharedRef<class IPrope
 	UE::StateTree::PropertyHelpers::GetStructValue<FGuid>(IDProperty, ID);
 
 	// Show the Value (FInstancedStruct) as child rows.
-	TSharedRef<FStateTreeStateParametersInstanceDataDetails> InstanceDetails = MakeShareable(new FStateTreeStateParametersInstanceDataDetails(ParametersValueProperty, ParametersProperty, PropUtils, bFixedLayout, ID, EditorData));
+	TSharedRef<FStateTreeStateParametersInstanceDataDetails> InstanceDetails = MakeShareable(new FStateTreeStateParametersInstanceDataDetails(ParametersProperty, PropUtils, bFixedLayout, ID, EditorData));
 	StructBuilder.AddCustomBuilder(InstanceDetails);
 }
 
diff --git a/Engine/Plugins/Runtime/StateTree/Source/StateTreeEditorModule/Private/Customizations/StateTreeStateParametersDetails.h b/Engine/Plugins/Runtime/StateTree/Source/StateTreeEditorModule/Private/Customizations/StateTreeStateParametersDetails.h
--- a/Engine/Plugins/Runtime/StateTree/Source/StateTreeEditorModule/Private/Customizations/StateTreeStateParametersDetails.h
+++ b/Engine/Plugins/Runtime/StateTree/Source/StateTreeEditorModule/Private/Customizations/StateTreeStateParametersDetails.h
@@ -29,7 +29,6 @@ private:
 	class IPropertyUtilities* PropUtils = nullptr;
 
 	TSharedPtr<IPropertyHandle> ParametersProperty;
-	TSharedPtr<IPropertyHandle> ParametersValueProperty;
 	TSharedPtr<IPropertyHandle> FixedLayoutProperty;
 	TSharedPtr<IPropertyHandle> IDProperty;
 	TSharedPtr<IPropertyHandle> StructProperty;
