Fixed StructUtils plugin's markup to properly reflect that it's still experimental in 5.1

diff --git a/Engine/Plugins/Experimental/StructUtils/StructUtils.uplugin b/Engine/Plugins/Experimental/StructUtils/StructUtils.uplugin
--- a/Engine/Plugins/Experimental/StructUtils/StructUtils.uplugin
+++ b/Engine/Plugins/Experimental/StructUtils/StructUtils.uplugin
@@ -12,7 +12,7 @@
 	"SupportURL": "",
 	"EnabledByDefault": false,
 	"CanContainContent": false,
-	"IsBetaVersion": true,
+	"IsBetaVersion": false,
 	"IsExperimentalVersion": true,
 	"Installed": false,
 	"Modules": [
