InstancedStruct: Fixed BP MakeInstancedStruct

diff --git a/Engine/Plugins/Experimental/StructUtils/Source/StructUtils/Private/StructUtilsFunctionLibrary.cpp b/Engine/Plugins/Experimental/StructUtils/Source/StructUtils/Private/StructUtilsFunctionLibrary.cpp
--- a/Engine/Plugins/Experimental/StructUtils/Source/StructUtils/Private/StructUtilsFunctionLibrary.cpp
+++ b/Engine/Plugins/Experimental/StructUtils/Source/StructUtils/Private/StructUtilsFunctionLibrary.cpp
@@ -28,8 +28,7 @@ void UStructUtilsFunctionLibrary::GetInstancedStructValue(EStructUtilsResult& Ex
 
 DEFINE_FUNCTION(UStructUtilsFunctionLibrary::execMakeInstancedStruct)
 {
-	Stack.Step(Stack.Object, nullptr);
-
+	// Read wildcard Value input.
 	Stack.MostRecentPropertyAddress = nullptr;
 	Stack.MostRecentPropertyContainer = nullptr;
 	Stack.StepCompiledIn<FStructProperty>(nullptr);
@@ -64,6 +63,7 @@ DEFINE_FUNCTION(UStructUtilsFunctionLibrary::execSetInstancedStructValue)
 {
 	P_GET_STRUCT_REF(FInstancedStruct, InstancedStruct);
 
+	// Read wildcard Value input.
 	Stack.MostRecentPropertyAddress = nullptr;
 	Stack.MostRecentPropertyContainer = nullptr;
 	Stack.StepCompiledIn<FStructProperty>(nullptr);
@@ -99,6 +99,7 @@ DEFINE_FUNCTION(UStructUtilsFunctionLibrary::execGetInstancedStructValue)
 	P_GET_ENUM_REF(EStructUtilsResult, ExecResult);
 	P_GET_STRUCT_REF(FInstancedStruct, InstancedStruct);
 
+	// Read wildcard Value input.
 	Stack.MostRecentPropertyAddress = nullptr;
 	Stack.MostRecentPropertyContainer = nullptr;
 	Stack.StepCompiledIn<FStructProperty>(nullptr);
diff --git a/Engine/Plugins/Experimental/StructUtils/Source/StructUtils/Public/InstancedStruct.h b/Engine/Plugins/Experimental/StructUtils/Source/StructUtils/Public/InstancedStruct.h
--- a/Engine/Plugins/Experimental/StructUtils/Source/StructUtils/Public/InstancedStruct.h
+++ b/Engine/Plugins/Experimental/StructUtils/Source/StructUtils/Public/InstancedStruct.h
@@ -18,7 +18,7 @@ struct FConstSharedStruct;
  *	UPROPERTY(EditAnywhere, Category = Foo, meta = (BaseStruct = "/Script/ModuleName.TestStructBase"))
  *	TArray<FInstancedStruct> TestArray;
  */
-USTRUCT(BlueprintType)
+USTRUCT(BlueprintType, meta = (HasNativeMake = "/Script/StructUtils.StructUtilsFunctionLibrary.MakeInstancedStruct"))
 struct STRUCTUTILS_API FInstancedStruct
 {
 	GENERATED_BODY()
diff --git a/Engine/Plugins/Experimental/StructUtils/Source/StructUtils/Public/StructUtilsFunctionLibrary.h b/Engine/Plugins/Experimental/StructUtils/Source/StructUtils/Public/StructUtilsFunctionLibrary.h
--- a/Engine/Plugins/Experimental/StructUtils/Source/StructUtils/Public/StructUtilsFunctionLibrary.h
+++ b/Engine/Plugins/Experimental/StructUtils/Source/StructUtils/Public/StructUtilsFunctionLibrary.h
@@ -25,7 +25,7 @@ public:
 	/**
 	 * Create a new InstancedStruct from the given source value.
 	 */
-	UFUNCTION(BlueprintCallable, CustomThunk, Category = "Utilities|Instanced Struct", meta = (CustomStructureParam = "Value", BlueprintInternalUseOnly="true"))
+	UFUNCTION(BlueprintCallable, CustomThunk, Category = "Utilities|Instanced Struct", meta = (CustomStructureParam = "Value", BlueprintInternalUseOnly="true", NativeMakeFunc))
 	static FInstancedStruct MakeInstancedStruct(const int32& Value);
 
 	/**
