PropetyBag: Change stray CastFieldChecked to CastField to prevent crashing.

diff --git a/Engine/Plugins/Experimental/StructUtils/Source/StructUtilsEditor/Private/Customizations/PropertyBagDetails.cpp b/Engine/Plugins/Experimental/StructUtils/Source/StructUtilsEditor/Private/Customizations/PropertyBagDetails.cpp
--- a/Engine/Plugins/Experimental/StructUtils/Source/StructUtilsEditor/Private/Customizations/PropertyBagDetails.cpp
+++ b/Engine/Plugins/Experimental/StructUtils/Source/StructUtilsEditor/Private/Customizations/PropertyBagDetails.cpp
@@ -38,7 +38,7 @@ bool IsScriptStruct(const TSharedPtr<IPropertyHandle>& PropertyHandle)
 		return false;
 	}
 
-	FStructProperty* StructProperty = CastFieldChecked<FStructProperty>(PropertyHandle->GetProperty());
+	FStructProperty* StructProperty = CastField<FStructProperty>(PropertyHandle->GetProperty());
 	return StructProperty && StructProperty->Struct == TBaseStructure<T>::Get();
 }
 
