StructUtils: Made AddStructObjectRefences to call AddStructObjectRefences recursively.
- This fixes GC on a struct that has a struct that has AddStructObjectRefences that is held by instanced struct (e.g. a struct that has FInstancedStruct in it)

diff --git a/Engine/Plugins/Experimental/StructUtils/Source/StructUtils/Private/InstancedStruct.cpp b/Engine/Plugins/Experimental/StructUtils/Source/StructUtils/Private/InstancedStruct.cpp
--- a/Engine/Plugins/Experimental/StructUtils/Source/StructUtils/Private/InstancedStruct.cpp
+++ b/Engine/Plugins/Experimental/StructUtils/Source/StructUtils/Private/InstancedStruct.cpp
@@ -1,7 +1,7 @@
 // Copyright Epic Games, Inc. All Rights Reserved.
 #include "InstancedStruct.h"
 #include "StructView.h"
-#include "SharedStruct.h"
+#include "StructUtilsTypes.h"
 #include "Serialization/PropertyLocalizationDataGathering.h"
 
 #include UE_INLINE_GENERATED_CPP_BY_NAME(InstancedStruct)
@@ -84,7 +84,7 @@ void FInstancedStruct::InitializeAs(const UScriptStruct* InScriptStruct, const u
 		{
 			const int32 MinAlignment = InScriptStruct->GetMinAlignment();
 			const int32 RequiredSize = InScriptStruct->GetStructureSize();
-			const uint8* Memory = ((uint8*)FMemory::Malloc(FMath::Max(1, RequiredSize), MinAlignment));
+			uint8* Memory = ((uint8*)FMemory::Malloc(FMath::Max(1, RequiredSize), MinAlignment));
 			SetStructData(InScriptStruct, Memory);
 
 			InScriptStruct->InitializeStruct(GetMutableMemory());
@@ -366,9 +366,9 @@ bool FInstancedStruct::Identical(const FInstancedStruct* Other, uint32 PortFlags
 
 void FInstancedStruct::AddStructReferencedObjects(class FReferenceCollector& Collector)
 {
-	if (const UScriptStruct* Struct = GetScriptStruct())
+	if (ScriptStruct != nullptr)
 	{
-		Collector.AddReferencedObjects(Struct, GetMutableMemory());
+		UE::StructUtils::AddReferencedObjects(Collector, ScriptStruct, GetMutableMemory());
 	}
 }
 
diff --git a/Engine/Plugins/Experimental/StructUtils/Source/StructUtils/Private/InstancedStructArray.cpp b/Engine/Plugins/Experimental/StructUtils/Source/StructUtils/Private/InstancedStructArray.cpp
--- a/Engine/Plugins/Experimental/StructUtils/Source/StructUtils/Private/InstancedStructArray.cpp
+++ b/Engine/Plugins/Experimental/StructUtils/Source/StructUtils/Private/InstancedStructArray.cpp
@@ -1,5 +1,6 @@
 // Copyright Epic Games, Inc. All Rights Reserved.
 #include "InstancedStructArray.h"
+#include "StructUtilsTypes.h"
 #include "Serialization/PropertyLocalizationDataGathering.h"
 
 #include UE_INLINE_GENERATED_CPP_BY_NAME(InstancedStructArray)
@@ -369,8 +370,7 @@ void FInstancedStructArray::AddStructReferencedObjects(class FReferenceCollector
 		FItem& Item = GetItem(Index);
 		if (Item.ScriptStruct != nullptr)
 		{
-			Collector.AddReferencedObject(Item.ScriptStruct);
-			Collector.AddReferencedObjects(Item.ScriptStruct, Memory + Item.Offset);
+			UE::StructUtils::AddReferencedObjects(Collector, Item.ScriptStruct, Memory + Item.Offset);
 		}
 	}
 }
diff --git a/Engine/Plugins/Experimental/StructUtils/Source/StructUtils/Private/StructUtilsTypes.cpp b/Engine/Plugins/Experimental/StructUtils/Source/StructUtils/Private/StructUtilsTypes.cpp
--- a/Engine/Plugins/Experimental/StructUtils/Source/StructUtils/Private/StructUtilsTypes.cpp
+++ b/Engine/Plugins/Experimental/StructUtils/Source/StructUtils/Private/StructUtilsTypes.cpp
@@ -3,7 +3,7 @@
 #include "StructUtilsTypes.h"
 #include "UObject/Class.h"
 #include "Serialization/ArchiveCrc32.h"
-#include "InstancedStruct.h"
+#include "UObject/UObjectGlobals.h"
 
 namespace UE::StructUtils
 {
@@ -27,4 +27,37 @@ namespace UE::StructUtils
 		return 0;
 	}
 
+	STRUCTUTILS_API void AddReferencedObjects(FReferenceCollector& Collector, const UScriptStruct*& ScriptStruct, void* StructMemory,
+										 const UObject* ReferencingObject, const FProperty* ReferencingProperty)
+	{
+		check(ScriptStruct != nullptr);
+		check(StructMemory != nullptr);
+
+		Collector.AddReferencedObject(ScriptStruct, ReferencingObject, ReferencingProperty);
+
+		// If the script struct explicitly provided an implementation of AddReferencedObjects, make sure to capture its referenced objects
+		if (ScriptStruct->StructFlags & STRUCT_AddStructReferencedObjects)
+		{
+			ScriptStruct->GetCppStructOps()->AddStructReferencedObjects()(StructMemory, Collector);
+		}
+
+		// Iterate through all object and struct properties within the struct (will also search through structs within the struct)
+		for (TPropertyValueIterator<const FProperty> PropertyIter(ScriptStruct, StructMemory); PropertyIter; ++PropertyIter)
+		{
+			if (const FObjectProperty* ObjectProperty = CastField<FObjectProperty>(PropertyIter.Key()))
+			{
+				TObjectPtr<UObject>& ObjectPtrRef = ObjectProperty->GetObjectPtrPropertyValueRef(PropertyIter.Value());
+				Collector.AddReferencedObject(ObjectPtrRef, ReferencingObject, ReferencingProperty);
+			}
+			else if (const FStructProperty* StructProperty = CastField<FStructProperty>(PropertyIter.Key()))
+			{
+				const UScriptStruct* ChildScriptStruct = StructProperty->Struct; 
+				if (ChildScriptStruct && ChildScriptStruct->StructFlags & STRUCT_AddStructReferencedObjects)
+				{
+					ChildScriptStruct->GetCppStructOps()->AddStructReferencedObjects()(const_cast<void*>(PropertyIter.Value()), Collector);
+				}
+			}
+		}
+	}
+
 }
diff --git a/Engine/Plugins/Experimental/StructUtils/Source/StructUtils/Public/InstancedStruct.h b/Engine/Plugins/Experimental/StructUtils/Source/StructUtils/Public/InstancedStruct.h
--- a/Engine/Plugins/Experimental/StructUtils/Source/StructUtils/Public/InstancedStruct.h
+++ b/Engine/Plugins/Experimental/StructUtils/Source/StructUtils/Public/InstancedStruct.h
@@ -64,7 +64,7 @@ public:
 		{
 			Reset();
 
-			SetStructData(InOther.GetScriptStruct(), InOther.GetMemory());
+			SetStructData(InOther.GetScriptStruct(), InOther.GetMutableMemory());
 			InOther.SetStructData(nullptr,nullptr);
 		}
 		return *this;
@@ -186,11 +186,10 @@ public:
 		return nullptr;
 	}
 
-	/** Returns a mutable pointer to struct memory. This const_cast here is safe as a ClassName can only be setup from mutable non const memory. */
+	/** Returns a mutable pointer to struct memory. */
 	uint8* GetMutableMemory() const
 	{
-		const uint8* Memory = GetMemory();
-		return const_cast<uint8*>(Memory);
+		return StructMemory;
 	}
 
 	/** Returns mutable reference to the struct, this getter assumes that all data is valid. */
@@ -263,7 +262,7 @@ protected:
 		}
 	}
 
-	FInstancedStruct(const UScriptStruct* InScriptStruct, const uint8* InStructMemory)
+	FInstancedStruct(const UScriptStruct* InScriptStruct, uint8* InStructMemory)
 		: ScriptStruct(InScriptStruct)
 		, StructMemory(InStructMemory)
 	{}
@@ -272,15 +271,14 @@ protected:
 		StructMemory = nullptr;
 		ScriptStruct = nullptr;
 	}
-	void SetStructData(const UScriptStruct* InScriptStruct, const uint8* InStructMemory)
+	void SetStructData(const UScriptStruct* InScriptStruct, uint8* InStructMemory)
 	{
 		ScriptStruct = InScriptStruct;
 		StructMemory = InStructMemory;
 	}
 
-
 	const UScriptStruct* ScriptStruct = nullptr;
-	const uint8* StructMemory = nullptr;
+	uint8* StructMemory = nullptr;
 };
 
 template<>
diff --git a/Engine/Plugins/Experimental/StructUtils/Source/StructUtils/Public/StructUtilsTypes.h b/Engine/Plugins/Experimental/StructUtils/Source/StructUtils/Public/StructUtilsTypes.h
--- a/Engine/Plugins/Experimental/StructUtils/Source/StructUtils/Public/StructUtilsTypes.h
+++ b/Engine/Plugins/Experimental/StructUtils/Source/StructUtils/Public/StructUtilsTypes.h
@@ -30,6 +30,13 @@ namespace UE::StructUtils
 	{
 		return T::StaticClass();
 	}
+
+	/**
+	 * Adds references objects from specific struct, similar to FReferenceCollector::AddReferencedObjects(),
+	 * but calls AddStructReferencedObjects recursively on child structs.
+	 */
+	extern STRUCTUTILS_API void AddReferencedObjects(FReferenceCollector& Collector, const UScriptStruct*& ScriptStruct, void* StructMemory,
+													 const UObject* ReferencingObject = nullptr, const FProperty* ReferencingProperty = nullptr);
 }
 
 /* Predicate useful to find a struct of a specific type in an container */
