PR #9280: NetSerialize / Replication support for FInstancedStruct (Contributed by [REDACTED_PERSON])

diff --git a/Engine/Plugins/Experimental/StructUtils/Source/StructUtils/Private/InstancedStruct.cpp b/Engine/Plugins/Experimental/StructUtils/Source/StructUtils/Private/InstancedStruct.cpp
--- a/Engine/Plugins/Experimental/StructUtils/Source/StructUtils/Private/InstancedStruct.cpp
+++ b/Engine/Plugins/Experimental/StructUtils/Source/StructUtils/Private/InstancedStruct.cpp
@@ -3,6 +3,9 @@
 #include "StructView.h"
 #include "StructUtilsTypes.h"
 #include "Serialization/PropertyLocalizationDataGathering.h"
+#include "Engine/PackageMapClient.h"
+#include "Engine/NetConnection.h"
+#include "Net/RepLayout.h"
 
 #include UE_INLINE_GENERATED_CPP_BY_NAME(InstancedStruct)
 
@@ -372,3 +375,74 @@ void FInstancedStruct::AddStructReferencedObjects(class FReferenceCollector& Col
 	}
 }
 
+bool FInstancedStruct::NetSerialize(FArchive& Ar, UPackageMap* Map, bool& bOutSuccess)
+{
+	uint8 bValidData = Ar.IsSaving() ? IsValid() : 0;
+	Ar.SerializeBits(&bValidData, 1);
+
+	if (bValidData)
+	{
+		if (Ar.IsLoading())
+		{
+			UScriptStruct* NonConstStruct = nullptr;
+
+			Ar << NonConstStruct;
+
+			// Initialize only if the type changes.
+			if (ScriptStruct != NonConstStruct)
+			{
+				InitializeAs(NonConstStruct);
+			}
+
+			if (!IsValid())
+			{
+				UE_LOG(LogCore, Error, TEXT("FInstancedStruct::NetSerialize: Bad script struct serialized, cannot recover."));
+				Ar.SetError();
+				bOutSuccess = false;
+			}
+		}
+		else if (Ar.IsSaving())
+		{
+			UScriptStruct* NonConstStruct = const_cast<UScriptStruct*>(ScriptStruct);
+			check(::IsValid(NonConstStruct));
+			
+			Ar << NonConstStruct;
+		}
+
+		// Check ScriptStruct here, as loading might have failed. 
+		if (ScriptStruct)
+		{
+			if (ScriptStruct->StructFlags & STRUCT_NetSerializeNative)
+			{
+				ScriptStruct->GetCppStructOps()->NetSerialize(Ar, Map, bOutSuccess, GetMutableMemory());
+			}
+			else
+			{
+				UPackageMapClient* MapClient = Cast<UPackageMapClient>(Map);
+				check(::IsValid(MapClient));
+
+				UNetConnection* NetConnection = MapClient->GetConnection();
+				check(::IsValid(NetConnection));
+				check(::IsValid(NetConnection->GetDriver()));
+
+				UScriptStruct* NonConstStruct = const_cast<UScriptStruct*>(ScriptStruct);
+				const TSharedPtr<FRepLayout> RepLayout = NetConnection->GetDriver()->GetStructRepLayout(NonConstStruct);
+				check(RepLayout.IsValid());
+
+				bool bHasUnmapped = false;
+				RepLayout->SerializePropertiesForStruct(NonConstStruct, static_cast<FBitArchive&>(Ar), Map, GetMutableMemory(), bHasUnmapped);
+				bOutSuccess = true;
+			}
+		}
+	}
+	else
+	{
+		if (Ar.IsLoading())
+		{
+			Reset();
+		}
+		bOutSuccess = true;
+	}
+
+	return true;
+}
\ No newline at end of file
diff --git a/Engine/Plugins/Experimental/StructUtils/Source/StructUtils/Public/InstancedStruct.h b/Engine/Plugins/Experimental/StructUtils/Source/StructUtils/Public/InstancedStruct.h
--- a/Engine/Plugins/Experimental/StructUtils/Source/StructUtils/Public/InstancedStruct.h
+++ b/Engine/Plugins/Experimental/StructUtils/Source/StructUtils/Public/InstancedStruct.h
@@ -145,6 +145,7 @@ public:
 	bool ImportTextItem(const TCHAR*& Buffer, int32 PortFlags, UObject* Parent, FOutputDevice* ErrorText, FArchive* InSerializingArchive = nullptr);
 	bool SerializeFromMismatchedTag(const FPropertyTag& Tag, FStructuredArchive::FSlot Slot);
 	void GetPreloadDependencies(TArray<UObject*>& OutDeps);
+	bool NetSerialize(FArchive& Ar, class UPackageMap* Map, bool& bOutSuccess);
 
 	/** Returns struct type. */
 	const UScriptStruct* GetScriptStruct() const
@@ -293,5 +294,6 @@ struct TStructOpsTypeTraits<FInstancedStruct> : public TStructOpsTypeTraitsBase2
 		WithAddStructReferencedObjects = true,
 		WithStructuredSerializeFromMismatchedTag = true,
 		WithGetPreloadDependencies = true,
+		WithNetSerializer = true,
 	};
 };
