StructUtils: Added ranged for support for instanced struct container

diff --git a/Engine/Plugins/AI/MassCrowd/Source/MassCrowd/Private/ZoneGraphCrowdLaneAnnotations.cpp b/Engine/Plugins/AI/MassCrowd/Source/MassCrowd/Private/ZoneGraphCrowdLaneAnnotations.cpp
--- a/Engine/Plugins/AI/MassCrowd/Source/MassCrowd/Private/ZoneGraphCrowdLaneAnnotations.cpp
+++ b/Engine/Plugins/AI/MassCrowd/Source/MassCrowd/Private/ZoneGraphCrowdLaneAnnotations.cpp
@@ -34,9 +34,8 @@ FZoneGraphTagMask UZoneGraphCrowdLaneAnnotations::GetAnnotationTags() const
 
 void UZoneGraphCrowdLaneAnnotations::HandleEvents(const FInstancedStructContainer& Events)
 {
-	for (int32 Index = 0; Index < Events.Num(); Index++)
+	for (FConstStructView Event : Events)
 	{
-		FConstStructView Event = Events[Index];
 		if (const FZoneGraphCrowdLaneStateChangeEvent* const StateChangeEvent = Event.GetPtr<FZoneGraphCrowdLaneStateChangeEvent>())
 		{
 			StateChangeEvents.Add(*StateChangeEvent);
diff --git a/Engine/Plugins/Experimental/StructUtils/Source/StructUtils/Public/InstancedStructContainer.h b/Engine/Plugins/Experimental/StructUtils/Source/StructUtils/Public/InstancedStructContainer.h
--- a/Engine/Plugins/Experimental/StructUtils/Source/StructUtils/Public/InstancedStructContainer.h
+++ b/Engine/Plugins/Experimental/StructUtils/Source/StructUtils/Public/InstancedStructContainer.h
@@ -4,6 +4,7 @@
 
 #include "InstancedStruct.h"
 #include "StructView.h"
+#include "Templates/ChooseClass.h"
 #include "InstancedStructContainer.generated.h"
 
 /**
@@ -125,6 +126,99 @@ public:
 		return FStructView(Item.ScriptStruct, Memory + Item.Offset);
 	}
 
+	/**
+	 * Iterators to enable range-based for loop support.
+	 *
+	 *	// Ranged for mutable container 
+	 *	for (FStructView View : Container) {}
+	 *
+	 *	// Ranged for const container 
+	 *	for (FConstStructView View : Container) {}
+	 *
+	 *	// Iterator based iteration, allows removing items.
+	 *	for (FInstancedStructContainer::FIterator It = Container.CreateIterator(); It; ++It)
+	 *	{
+	 *		It.RemoveCurrent();
+	 *	}
+	 */
+	template<typename T>
+	struct TIterator
+	{
+		using StructViewType = typename TChooseClass<TIsConst<T>::Value, FConstStructView, FStructView>::Result;
+
+		/** @return struct view (or const structview) to the item. */
+		StructViewType operator*()
+		{
+			const FItem& Item = Container.GetItem(Index);
+			return StructViewType(Item.ScriptStruct, Container.Memory + Item.Offset);
+		}
+
+		TIterator& operator++()
+		{
+			Index++;
+			return *this;
+		}
+
+		TIterator& operator--()
+		{
+			Index--;
+			return *this;
+		}
+
+		/** @return true of two iterators point to the same item. */
+		bool operator!=(const TIterator& RHS) const
+		{
+			return Index != RHS.Index;
+		}
+
+		/** @return true of the interator points to valid item. */
+		explicit operator bool() const
+		{
+			return Container.IsValidIndex(Index);
+		}
+
+		/** @returns The index of the iterator.  */
+		int32 GetIndex() const
+		{
+			return Index;
+		}
+
+		/** Removes the item pointed by the iterator and adjust the iterator. */
+		void RemoveCurrent()
+		{
+			static_assert(!TIsConst<T>::Value, "Cannot remove on const container.");
+			Container.RemoveAt(Index, 1);
+			Index--;
+		}
+		
+	private:
+		explicit TIterator(T& InContainer, const int32 InIndex = 0)
+			: Container(InContainer)
+			, Index(InIndex)
+		{
+		}
+
+		T& Container;
+		int32 Index = 0;
+
+		friend FInstancedStructContainer;
+	};
+
+	using FIterator = TIterator<FInstancedStructContainer>;
+	using FConstIterator = TIterator<const FInstancedStructContainer>;
+
+	/** Creates iterator to iterate over the array. */
+	FIterator CreateIterator() { return FIterator(*this, 0); }
+	/** Creates const iterator to iterate over the array. */
+	FConstIterator CreateConstIterator() const { return FConstIterator(*this, 0); }
+
+	/** For ranged for, do not use directly. */
+	FIterator begin() { return FIterator(*this, 0); }
+	FIterator end() { return FIterator(*this, NumItems); }
+
+	FConstIterator begin() const { return FConstIterator(*this, 0); }
+	FConstIterator end() const { return FConstIterator(*this, NumItems); }
+
 	/** Type traits */
 	void AddStructReferencedObjects(class FReferenceCollector& Collector) const;
 	bool Identical(const FInstancedStructContainer* Other, uint32 PortFlags) const;
@@ -162,6 +256,9 @@ private:
 
 	/** Number of items in the array. */
 	int32 NumItems = 0;
+
+	friend FIterator;
+	friend FConstIterator;
 };
 
 template<>
diff --git a/Engine/Plugins/Experimental/StructUtils/Source/StructUtilsTestSuite/Private/InstancedStructContainerTest.cpp b/Engine/Plugins/Experimental/StructUtils/Source/StructUtilsTestSuite/Private/InstancedStructContainerTest.cpp
--- a/Engine/Plugins/Experimental/StructUtils/Source/StructUtilsTestSuite/Private/InstancedStructContainerTest.cpp
+++ b/Engine/Plugins/Experimental/StructUtils/Source/StructUtilsTestSuite/Private/InstancedStructContainerTest.cpp
@@ -27,7 +27,7 @@ struct FTest_BasicInstancedStructContainer : FAITestBase
 
 		Array[1].GetMutable<FVector>().X = 42.0;
 
-		AITEST_TRUE(TEXT("Should have 3 items"), Array.Num() == 3);
+		AITEST_EQUAL(TEXT("Should have 3 items"), Array.Num(), 3);
 		AITEST_TRUE(TEXT("Item 0 should be FTestStructSimple"), Array[0].GetScriptStruct() == TBaseStructure<FTestStructSimple>::Get());
 		AITEST_TRUE(TEXT("Item 1 should be FVector"), Array[1].GetScriptStruct() == TBaseStructure<FVector>::Get());
 		AITEST_TRUE(TEXT("Item 2 should be FTestStructComplex"), Array[2].GetScriptStruct() == TBaseStructure<FTestStructComplex>::Get());
@@ -40,23 +40,23 @@ struct FTest_BasicInstancedStructContainer : FAITestBase
 		TestInstanced.AddDefaulted_GetRef().InitializeAs<FVector>(1,2,3);
 		Array.Append(TestInstanced);
 
-		AITEST_TRUE(TEXT("Should have 5 items"), Array.Num() == 5);
+		AITEST_EQUAL(TEXT("Should have 5 items"), Array.Num(), 5);
 		AITEST_TRUE(TEXT("Item 4 should have Z == 3.0"), FMath::IsNearlyEqual(Array[4].GetMutable<FVector>().Z, 3.0));
 		AITEST_TRUE(TEXT("Item 2 should have text Foo"), Array[2].GetMutable<FTestStructComplex>().StringArray[0] == TEXT("Foo"));
 
 		Array.RemoveAt(2, 1);
-		AITEST_TRUE(TEXT("Should have 4 items"), Array.Num() == 4);
+		AITEST_EQUAL(TEXT("Should have 4 items"), Array.Num(), 4);
 		AITEST_TRUE(TEXT("Item 2 should be FTransform"), Array[2].GetScriptStruct() == TBaseStructure<FTransform>::Get());
 		AITEST_TRUE(TEXT("Item 3 should be FVector"), Array[3].GetScriptStruct() == TBaseStructure<FVector>::Get());
 		AITEST_TRUE(TEXT("Item 3 should have Z == 3.0"), FMath::IsNearlyEqual(Array[3].GetMutable<FVector>().Z, 3.0));
 
 		Array.SetNum(2);
-		AITEST_TRUE(TEXT("Should have 2 items"), Array.Num() == 2);
+		AITEST_EQUAL(TEXT("Should have 2 items"), Array.Num(), 2);
 
 		FInstancedStructContainer Array2;
 		Array2.Append({ TBaseStructure<FQuat>::Get(), TBaseStructure<FTestStructSimple1>::Get() });
 		Array.InsertAt(1, Array2);
-		AITEST_TRUE(TEXT("Should have 4 items"), Array.Num() == 4);
+		AITEST_EQUAL(TEXT("Should have 4 items"), Array.Num(), 4);
 		AITEST_TRUE(TEXT("Item 0 should be FTestStructSimple"), Array[0].GetScriptStruct() == TBaseStructure<FTestStructSimple>::Get());
 		AITEST_TRUE(TEXT("Item 1 should be FQuat"), Array[1].GetScriptStruct() == TBaseStructure<FQuat>::Get());
 		AITEST_TRUE(TEXT("Item 2 should be FTestStructSimple1"), Array[2].GetScriptStruct() == TBaseStructure<FTestStructSimple1>::Get());
@@ -64,12 +64,12 @@ struct FTest_BasicInstancedStructContainer : FAITestBase
 		AITEST_TRUE(TEXT("Item 3 should have X == 42.0"), FMath::IsNearlyEqual(Array[3].GetMutable<FVector>().X, 42.0));
 
 		Array.RemoveAt(0, 3);
-		AITEST_TRUE(TEXT("Should have 1 item"), Array.Num() == 1);
+		AITEST_EQUAL(TEXT("Should have 1 item"), Array.Num(), 1);
 		AITEST_TRUE(TEXT("Item 0 should be FVector"), Array[0].GetScriptStruct() == TBaseStructure<FVector>::Get());
 		AITEST_TRUE(TEXT("Item 0 should have X == 42.0"), FMath::IsNearlyEqual(Array[0].GetMutable<FVector>().X, 42.0));
 
 		Array.InsertAt(0, Array2);
-		AITEST_TRUE(TEXT("Should have 1 item"), Array.Num() == 3);
+		AITEST_EQUAL(TEXT("Should have 1 item"), Array.Num(), 3);
 		AITEST_TRUE(TEXT("Item 2 should have X == 42.0"), FMath::IsNearlyEqual(Array[2].GetMutable<FVector>().X, 42.0));
 
 		return true;
@@ -91,7 +91,7 @@ struct FTest_SerializeInstancedStructContainer : FAITestBase
 
 		Array[2].GetMutable<FTestStructComplex>().StringArray.Add(TEXT("Foo"));
 
-		AITEST_TRUE(TEXT("Should have 3 items"), Array.Num() == 3);
+		AITEST_EQUAL(TEXT("Should have 3 items"), Array.Num(), 3);
 		AITEST_TRUE(TEXT("Item 2 should have text Foo"), Array[2].GetMutable<FTestStructComplex>().StringArray[0] == TEXT("Foo"));
 
 		TArray<uint8> Memory;
@@ -111,7 +111,7 @@ struct FTest_SerializeInstancedStructContainer : FAITestBase
 		FObjectAndNameAsStringProxyArchive ReaderProxy2(Reader2, /*bInLoadIfFindFails*/true);
 		const bool bLoadResult2 = Array2.Serialize(ReaderProxy2);
 		AITEST_TRUE(TEXT("Loading to Array2 should succeed"), bLoadResult2);
-		AITEST_TRUE(TEXT("Array2 should have 3 items"), Array2.Num() == 3);
+		AITEST_EQUAL(TEXT("Array2 should have 3 items"), Array2.Num(), 3);
 		AITEST_TRUE(TEXT("Array2 item 2 should have text Foo"), Array2[2].GetMutable<FTestStructComplex>().StringArray[0] == TEXT("Foo"));
 
 		return true;
@@ -119,6 +119,68 @@ struct FTest_SerializeInstancedStructContainer : FAITestBase
 };
 IMPLEMENT_AI_INSTANT_TEST(FTest_SerializeInstancedStructContainer, "System.StructUtils.InstancedStructContainer.Serialize");
 
+struct FTest_RangedForInstancedStructContainer : FAITestBase
+{
+	virtual bool InstantTest() override
+	{
+		FInstancedStructContainer Array;
+		const FInstancedStructContainer& ConstArray = Array;
+
+		TArray<FInstancedStruct> TestInstanced;
+		TestInstanced.AddDefaulted_GetRef().InitializeAs<FTestStructComplex>(TEXT("0"));
+		TestInstanced.AddDefaulted_GetRef().InitializeAs<FTestStructComplex>(TEXT("1"));
+		TestInstanced.AddDefaulted_GetRef().InitializeAs<FTestStructComplex>(TEXT("2"));
+		Array.Append(TestInstanced);
+
+		int32 Count = 0;
+		for (FStructView Item : Array)
+		{
+			if (Item.GetMutablePtr<FTestStructComplex>())
+			{
+				Count++;
+			}
+		}
+		AITEST_EQUAL(TEXT("Count should be 3"), Count, 3);
+
+		for (FConstStructView Item : ConstArray)
+		{
+			if (Item.GetPtr<FTestStructComplex>())
+			{
+				Count++;
+			}
+		}
+		AITEST_EQUAL(TEXT("Count should be 6"), Count, 6);
+
+		int32 ItemIndex = 0;
+		for (FConstStructView Item : ConstArray)
+		{
+			AITEST_EQUAL("Should iterate in order", Item, ConstArray[ItemIndex++]);
+		}
+
+		FInstancedStructContainer::FIterator Iter = Array.CreateIterator();
+		++Iter;
+		Iter.RemoveCurrent();
+		AITEST_EQUAL(TEXT("Array should have 2 items"), Array.Num(), 2);
+		AITEST_EQUAL(TEXT("Iter should be at 0"), Iter.GetIndex(), 0);
+		AITEST_TRUE(TEXT("Iter should be true"), (bool)Iter);
+		AITEST_EQUAL(TEXT("Array[0] should be 0"), Array[0].Get<FTestStructComplex>().String, TEXT("0"));
+		AITEST_EQUAL(TEXT("Array[1] should be 2"), Array[1].Get<FTestStructComplex>().String, TEXT("2"));
+
+		++Iter;
+		++Iter;
+		AITEST_FALSE(TEXT("Iter should be false"), (bool)Iter);
+
+		for (FInstancedStructContainer::FIterator Iter2 = Array.CreateIterator(); Iter2; ++Iter2)
+		{
+			Iter2.RemoveCurrent();
+		}
+		AITEST_EQUAL(TEXT("Array should have 0 items"), Array.Num(), 0);
+
+		return true;
+	}
+};
+IMPLEMENT_AI_INSTANT_TEST(FTest_RangedForInstancedStructContainer, "System.StructUtils.InstancedStructContainer.RangedFor");
+
 } // FInstancedStructContainerTest
 
 #undef LOCTEXT_NAMESPACE
diff --git a/Engine/Plugins/Runtime/ZoneGraphAnnotations/Source/ZoneGraphAnnotations/Private/Annotations/ZoneGraphDisturbanceAnnotation.cpp b/Engine/Plugins/Runtime/ZoneGraphAnnotations/Source/ZoneGraphAnnotations/Private/Annotations/ZoneGraphDisturbanceAnnotation.cpp
--- a/Engine/Plugins/Runtime/ZoneGraphAnnotations/Source/ZoneGraphAnnotations/Private/Annotations/ZoneGraphDisturbanceAnnotation.cpp
+++ b/Engine/Plugins/Runtime/ZoneGraphAnnotations/Source/ZoneGraphAnnotations/Private/Annotations/ZoneGraphDisturbanceAnnotation.cpp
@@ -572,10 +572,8 @@ void UZoneGraphDisturbanceAnnotation::UpdateAnnotationTags(FZoneGraphAnnotationT
 
 void UZoneGraphDisturbanceAnnotation::HandleEvents(const FInstancedStructContainer& Events)
 {
-	for (int32 Index = 0; Index < Events.Num(); Index++)
+	for (FConstStructView Event : Events)
 	{
-		FConstStructView Event = Events[Index];
-
 		if (const FZoneGraphDisturbanceArea* Danger = Event.GetPtr<FZoneGraphDisturbanceArea>())
 		{
 			FZoneGraphDisturbanceArea* ExistingDanger = Dangers.FindByPredicate([Danger](const FZoneGraphDisturbanceArea& Area) { return Area.InstigatorID == Danger->InstigatorID; });
