StructUtils: Fix crash in propertybag details
- SPinTypeSelector can outlive the detail rows of proeprty bag customization
- Make sure we can detect that case and fail gracefully

diff --git a/Engine/Plugins/Experimental/StructUtils/Source/StructUtilsEditor/Private/Customizations/PropertyBagDetails.cpp b/Engine/Plugins/Experimental/StructUtils/Source/StructUtilsEditor/Private/Customizations/PropertyBagDetails.cpp
--- a/Engine/Plugins/Experimental/StructUtils/Source/StructUtilsEditor/Private/Customizations/PropertyBagDetails.cpp
+++ b/Engine/Plugins/Experimental/StructUtils/Source/StructUtilsEditor/Private/Customizations/PropertyBagDetails.cpp
@@ -159,6 +159,16 @@ FName GetValidPropertyName(const FString& Name)
 	return Result;
 }
 
+FName GetPropertyNameSafe(const TSharedPtr<IPropertyHandle>& PropertyHandle)
+{
+	const FProperty* Property = PropertyHandle ? PropertyHandle->GetProperty() : nullptr;
+	if (Property != nullptr)
+	{
+		return Property->GetFName();
+	}
+	return FName();
+}
+
 /** @return true of the property name is not used yet by the property bag structure common to all edited properties. */
 bool IsUniqueName(const FName NewName, const FName OldName, const TSharedPtr<IPropertyHandle>& StructProperty)
 {
@@ -167,6 +177,11 @@ bool IsUniqueName(const FName NewName, const FName OldName, const TSharedPtr<IPr
 		return false;
 	}
 	
+	if (!StructProperty || !StructProperty->IsValidHandle())
+	{
+		return false;
+	}
+	
 	bool bFound = false;
 
 	if (ensure(IsScriptStruct<FInstancedPropertyBag>(StructProperty)))
@@ -392,8 +407,13 @@ void SetPropertyDescFromPin(FPropertyBagPropertyDesc& Desc, const FEdGraphPinTyp
 }
 
 template<typename TFunc>
-void ApplyChangesToPropertyDescs(const FText& SessionName, const TSharedPtr<IPropertyHandle>& StructProperty, IPropertyUtilities* PropUtils, TFunc&& Function)
+void ApplyChangesToPropertyDescs(const FText& SessionName, const TSharedPtr<IPropertyHandle>& StructProperty, const TSharedPtr<IPropertyUtilities>& PropUtils, TFunc&& Function)
 {
+	if (!StructProperty || !PropUtils)
+	{
+		return;
+	}
+	
 	FScopedTransaction Transaction(SessionName);
 	TArray<FPropertyBagPropertyDesc> PropertyDescs = GetCommonPropertyDescs(StructProperty);
 	StructProperty->NotifyPreChange();
@@ -443,7 +463,7 @@ bool CanHaveMemberVariableOfType(const FEdGraphPinType& PinType)
 //  - ChildPropertyHandle a child property of the FInstancedPropertyBag::Value (FInstancedStruct)  
 //----------------------------------------------------------------//
 
-FPropertyBagInstanceDataDetails::FPropertyBagInstanceDataDetails(TSharedPtr<IPropertyHandle> InStructProperty, IPropertyUtilities* InPropUtils, const bool bInFixedLayout)
+FPropertyBagInstanceDataDetails::FPropertyBagInstanceDataDetails(TSharedPtr<IPropertyHandle> InStructProperty, const TSharedPtr<IPropertyUtilities>& InPropUtils, const bool bInFixedLayout)
 	: FInstancedStructDataDetails(InStructProperty.IsValid() ? InStructProperty->GetChildHandle(TEXT("Value")) : nullptr)
 	, BagStructProperty(InStructProperty)
 	, PropUtils(InPropUtils)
@@ -470,19 +490,20 @@ void FPropertyBagInstanceDataDetails::OnChildRowAdded(IDetailPropertyRow& ChildR
 			.MultiLine(false)
 			.Text_Lambda([ChildPropertyHandle]()
 			{
-				return FText::FromName(ChildPropertyHandle->GetProperty()->GetFName());
+				return FText::FromName(UE::StructUtils::Private::GetPropertyNameSafe(ChildPropertyHandle));
 			})
-			.OnVerifyTextChanged_Lambda([this, ChildPropertyHandle](const FText& InText, FText& OutErrorMessage)
+			.OnVerifyTextChanged_Lambda([BagStructProperty = BagStructProperty, ChildPropertyHandle](const FText& InText, FText& OutErrorMessage)
 			{
 				const FName NewName = UE::StructUtils::Private::GetValidPropertyName(InText.ToString());
-				bool bResult = UE::StructUtils::Private::IsUniqueName(NewName, ChildPropertyHandle->GetProperty()->GetFName(), BagStructProperty);
+				const FName OldName = UE::StructUtils::Private::GetPropertyNameSafe(ChildPropertyHandle);
+				bool bResult = UE::StructUtils::Private::IsUniqueName(NewName, OldName, BagStructProperty);
 				if (!bResult)
 				{
 					OutErrorMessage = LOCTEXT("MustBeUniqueName", "Property must have unique name");
 				}
 				return bResult;
 			})
-			.OnTextCommitted_Lambda([this, ChildPropertyHandle](const FText& InNewText, ETextCommit::Type InCommitType)
+			.OnTextCommitted_Lambda([BagStructProperty = BagStructProperty, PropUtils = PropUtils, ChildPropertyHandle](const FText& InNewText, ETextCommit::Type InCommitType)
 			{
 				if (InCommitType == ETextCommit::OnCleared)
 				{
@@ -490,7 +511,8 @@ void FPropertyBagInstanceDataDetails::OnChildRowAdded(IDetailPropertyRow& ChildR
 				}
 
 				const FName NewName = UE::StructUtils::Private::GetValidPropertyName(InNewText.ToString());
-				if (!UE::StructUtils::Private::IsUniqueName(NewName, ChildPropertyHandle->GetProperty()->GetFName(), BagStructProperty))
+				const FName OldName = UE::StructUtils::Private::GetPropertyNameSafe(ChildPropertyHandle);
+				if (!UE::StructUtils::Private::IsUniqueName(NewName, OldName, BagStructProperty))
 				{
 					return;
 				}
@@ -535,7 +557,10 @@ void FPropertyBagInstanceDataDetails::OnChildRowAdded(IDetailPropertyRow& ChildR
 				[
 					SNew(SImage)
 					.ToolTipText(LOCTEXT("MissingType", "The property is missing type. The Struct, Enum, or Object may have been removed."))
-					.Visibility_Lambda([this, ChildPropertyHandle]() { return UE::StructUtils::Private::HasMissingType(ChildPropertyHandle) ? EVisibility::Visible : EVisibility::Collapsed; })
+					.Visibility_Lambda([ChildPropertyHandle]()
+					{
+						return UE::StructUtils::Private::HasMissingType(ChildPropertyHandle) ? EVisibility::Visible : EVisibility::Collapsed;
+					})
 					.Image(FAppStyle::GetBrush("Icons.Error"))
 				]
 			]
@@ -586,8 +611,14 @@ TSharedRef<SWidget> FPropertyBagInstanceDataDetails::OnPropertyNameContent(TShar
 		}			
 	};
 
-	auto GetPinInfo = [this, ChildPropertyHandle]()
+	auto GetPinInfo = [BagStructProperty = BagStructProperty, ChildPropertyHandle]()
 	{
+		// The SPinTypeSelector popup might outlive this details view, so bag struct property can be invalid here.
+		if (!BagStructProperty || !BagStructProperty->IsValidHandle() || !ChildPropertyHandle || !ChildPropertyHandle->IsValidHandle())
+		{
+			return FEdGraphPinType();
+		}
+		
 		TArray<FPropertyBagPropertyDesc> PropertyDescs = UE::StructUtils::Private::GetCommonPropertyDescs(BagStructProperty);
 
 		const FProperty* Property = ChildPropertyHandle->GetProperty();
@@ -599,14 +630,20 @@ TSharedRef<SWidget> FPropertyBagInstanceDataDetails::OnPropertyNameContent(TShar
 		return FEdGraphPinType();
 	};
 
-	auto PinInfoChanged = [this, ChildPropertyHandle](const FEdGraphPinType& PinType)
+	auto PinInfoChanged = [BagStructProperty = BagStructProperty, PropUtils = PropUtils, ChildPropertyHandle](const FEdGraphPinType& PinType)
 	{
+		// The SPinTypeSelector popup might outlive this details view, so bag struct property can be invalid here.
+		if (!BagStructProperty || !BagStructProperty->IsValidHandle() || !ChildPropertyHandle || !ChildPropertyHandle->IsValidHandle())
+		{
+			return;
+		}
+		
 		UE::StructUtils::Private::ApplyChangesToPropertyDescs(
 			LOCTEXT("OnPropertyTypeChanged", "Change Property Type"), BagStructProperty, PropUtils,
 			[&PinType, &ChildPropertyHandle](TArray<FPropertyBagPropertyDesc>& PropertyDescs)
 			{
 				// Find and change struct type
-				const FProperty* Property = ChildPropertyHandle->GetProperty();
+				const FProperty* Property = ChildPropertyHandle ? ChildPropertyHandle->GetProperty() : nullptr;
 				if (FPropertyBagPropertyDesc* Desc = PropertyDescs.FindByPredicate([Property](const FPropertyBagPropertyDesc& Desc){ return Desc.CachedProperty == Property; }))
 				{
 					UE::StructUtils::Private::SetPropertyDescFromPin(*Desc, PinType);
@@ -614,19 +651,29 @@ TSharedRef<SWidget> FPropertyBagInstanceDataDetails::OnPropertyNameContent(TShar
 			});
 	};
 
-	auto RemoveProperty = [this, ChildPropertyHandle]()
+	auto RemoveProperty = [BagStructProperty = BagStructProperty, PropUtils = PropUtils, ChildPropertyHandle]()
 	{
+		if (!BagStructProperty || !BagStructProperty->IsValidHandle() || !ChildPropertyHandle || !ChildPropertyHandle->IsValidHandle())
+		{
+			return;
+		}
+
 		UE::StructUtils::Private::ApplyChangesToPropertyDescs(
 			LOCTEXT("OnPropertyRemoved", "Remove Property"), BagStructProperty, PropUtils,
 			[&ChildPropertyHandle](TArray<FPropertyBagPropertyDesc>& PropertyDescs)
 			{
-				const FProperty* Property = ChildPropertyHandle->GetProperty();
+				const FProperty* Property = ChildPropertyHandle ? ChildPropertyHandle->GetProperty() : nullptr;
 				PropertyDescs.RemoveAll([Property](const FPropertyBagPropertyDesc& Desc){ return Desc.CachedProperty == Property; });
 			});
 	};
 
-	auto MoveProperty = [this, ChildPropertyHandle](const int32 Delta)
+	auto MoveProperty = [BagStructProperty = BagStructProperty, PropUtils = PropUtils, ChildPropertyHandle](const int32 Delta)
 	{
+		if (!BagStructProperty || !BagStructProperty->IsValidHandle() || !ChildPropertyHandle || !ChildPropertyHandle->IsValidHandle())
+		{
+			return;
+		}
+		
 		UE::StructUtils::Private::ApplyChangesToPropertyDescs(
 		LOCTEXT("OnPropertyMoved", "Move Property"), BagStructProperty, PropUtils,
 		[&ChildPropertyHandle, &Delta](TArray<FPropertyBagPropertyDesc>& PropertyDescs)
@@ -634,7 +681,7 @@ TSharedRef<SWidget> FPropertyBagInstanceDataDetails::OnPropertyNameContent(TShar
 			// Move
 			if (PropertyDescs.Num() > 1)
 			{
-				const FProperty* Property = ChildPropertyHandle->GetProperty();
+				const FProperty* Property = ChildPropertyHandle ? ChildPropertyHandle->GetProperty() : nullptr;
 				const int32 PropertyIndex = PropertyDescs.IndexOfByPredicate([Property](const FPropertyBagPropertyDesc& Desc){ return Desc.CachedProperty == Property; });
 				if (PropertyIndex != INDEX_NONE)
 				{
@@ -704,9 +751,9 @@ TSharedRef<IPropertyTypeCustomization> FPropertyBagDetails::MakeInstance()
 	return MakeShared<FPropertyBagDetails>();
 }
 
-void FPropertyBagDetails::CustomizeHeader(TSharedRef<class IPropertyHandle> StructPropertyHandle, class FDetailWidgetRow& HeaderRow, IPropertyTypeCustomizationUtils& StructCustomizationUtils)
+void FPropertyBagDetails::CustomizeHeader(TSharedRef<IPropertyHandle> StructPropertyHandle, FDetailWidgetRow& HeaderRow, IPropertyTypeCustomizationUtils& StructCustomizationUtils)
 {
-	PropUtils = StructCustomizationUtils.GetPropertyUtilities().Get();
+	PropUtils = StructCustomizationUtils.GetPropertyUtilities();
 	
 	StructProperty = StructPropertyHandle;
 	check(StructProperty);
@@ -736,14 +783,14 @@ void FPropertyBagDetails::CustomizeHeader(TSharedRef<class IPropertyHandle> Stru
 		.ShouldAutoExpand(true);
 }
 
-void FPropertyBagDetails::CustomizeChildren(TSharedRef<class IPropertyHandle> StructPropertyHandle, class IDetailChildrenBuilder& StructBuilder, IPropertyTypeCustomizationUtils& StructCustomizationUtils)
+void FPropertyBagDetails::CustomizeChildren(TSharedRef<IPropertyHandle> StructPropertyHandle, IDetailChildrenBuilder& StructBuilder, IPropertyTypeCustomizationUtils& StructCustomizationUtils)
 {
 	// Show the Value (FInstancedStruct) as child rows.
-	TSharedRef<FPropertyBagInstanceDataDetails> InstanceDetails = MakeShareable(new FPropertyBagInstanceDataDetails(StructProperty, PropUtils, bFixedLayout));
+	const TSharedRef<FPropertyBagInstanceDataDetails> InstanceDetails = MakeShareable(new FPropertyBagInstanceDataDetails(StructProperty, PropUtils, bFixedLayout));
 	StructBuilder.AddCustomBuilder(InstanceDetails);
 }
 
-TSharedPtr<SWidget> FPropertyBagDetails::MakeAddPropertyWidget(TSharedPtr<IPropertyHandle> InStructProperty, class IPropertyUtilities* InPropUtils)
+TSharedPtr<SWidget> FPropertyBagDetails::MakeAddPropertyWidget(TSharedPtr<IPropertyHandle> InStructProperty, TSharedPtr<IPropertyUtilities> InPropUtils)
 {
 	return SNew(SHorizontalBox)
 		+SHorizontalBox::Slot()
@@ -754,12 +801,12 @@ TSharedPtr<SWidget> FPropertyBagDetails::MakeAddPropertyWidget(TSharedPtr<IPrope
 			.HAlign(HAlign_Center)
 			.ButtonStyle(FAppStyle::Get(), "SimpleButton")
 			.ToolTipText(LOCTEXT("AddProperty_Tooltip", "Add new property"))
-			.OnClicked_Lambda([InStructProperty, InPropUtils]()
+			.OnClicked_Lambda([StructProperty = InStructProperty, PropUtils = InPropUtils]()
 			{
 				constexpr int32 MaxIterations = 100;
 				FName NewName(TEXT("NewProperty"));
 				int32 Number = 1;
-				while (!UE::StructUtils::Private::IsUniqueName(NewName, FName(), InStructProperty) && Number < MaxIterations)
+				while (!UE::StructUtils::Private::IsUniqueName(NewName, FName(), StructProperty) && Number < MaxIterations)
 				{
 					Number++;
 					NewName.SetNumber(Number);
@@ -770,7 +817,7 @@ TSharedPtr<SWidget> FPropertyBagDetails::MakeAddPropertyWidget(TSharedPtr<IPrope
 				}
 
 				UE::StructUtils::Private::ApplyChangesToPropertyDescs(
-					LOCTEXT("OnPropertyAdded", "Add Property"), InStructProperty, InPropUtils,
+					LOCTEXT("OnPropertyAdded", "Add Property"), StructProperty, PropUtils,
 					[&NewName](TArray<FPropertyBagPropertyDesc>& PropertyDescs)
 					{
 						PropertyDescs.Emplace(NewName, EPropertyBagPropertyType::Bool);
diff --git a/Engine/Plugins/Experimental/StructUtils/Source/StructUtilsEditor/Public/PropertyBagDetails.h b/Engine/Plugins/Experimental/StructUtils/Source/StructUtilsEditor/Public/PropertyBagDetails.h
--- a/Engine/Plugins/Experimental/StructUtils/Source/StructUtilsEditor/Public/PropertyBagDetails.h
+++ b/Engine/Plugins/Experimental/StructUtils/Source/StructUtilsEditor/Public/PropertyBagDetails.h
@@ -8,6 +8,7 @@
 #include "PropertyBagDetails.generated.h"
 
 class IPropertyHandle;
+class IPropertyUtilities;
 class IDetailPropertyRow;
 class SInlineEditableTextBlock;
 class SWidget;
@@ -22,12 +23,12 @@ public:
 	static TSharedRef<IPropertyTypeCustomization> MakeInstance();
 
 	/** Creates add property widget. */
-	static TSharedPtr<SWidget> MakeAddPropertyWidget(TSharedPtr<IPropertyHandle> StructProperty, class IPropertyUtilities* PropUtils);
+	static TSharedPtr<SWidget> MakeAddPropertyWidget(TSharedPtr<IPropertyHandle> InStructProperty, TSharedPtr<IPropertyUtilities> InPropUtils);
 	
 protected:
 	/** IPropertyTypeCustomization interface */
-	virtual void CustomizeHeader(TSharedRef<IPropertyHandle> StructPropertyHandle, class FDetailWidgetRow& HeaderRow, IPropertyTypeCustomizationUtils& StructCustomizationUtils) override;
-	virtual void CustomizeChildren(TSharedRef<IPropertyHandle> StructPropertyHandle, class IDetailChildrenBuilder& StructBuilder, IPropertyTypeCustomizationUtils& StructCustomizationUtils) override;
+	virtual void CustomizeHeader(TSharedRef<IPropertyHandle> StructPropertyHandle, FDetailWidgetRow& HeaderRow, IPropertyTypeCustomizationUtils& StructCustomizationUtils) override;
+	virtual void CustomizeChildren(TSharedRef<IPropertyHandle> StructPropertyHandle, IDetailChildrenBuilder& StructBuilder, IPropertyTypeCustomizationUtils& StructCustomizationUtils) override;
 
 private:
 
@@ -35,10 +36,8 @@ private:
 	
 	/** Handle to the struct property being edited */
 	TSharedPtr<IPropertyHandle> StructProperty;
-
+	TSharedPtr<IPropertyUtilities> PropUtils;
 	bool bFixedLayout = false; 
-	
-	class IPropertyUtilities* PropUtils = nullptr;
 };
 
 /** 
@@ -52,7 +51,7 @@ private:
 class STRUCTUTILSEDITOR_API FPropertyBagInstanceDataDetails : public FInstancedStructDataDetails
 {
 public:
-	FPropertyBagInstanceDataDetails(TSharedPtr<IPropertyHandle> InStructProperty, IPropertyUtilities* InPropUtils, const bool bInFixedLayout);
+	FPropertyBagInstanceDataDetails(TSharedPtr<IPropertyHandle> InStructProperty, const TSharedPtr<IPropertyUtilities>& InPropUtils, const bool bInFixedLayout);
 
 	virtual void OnChildRowAdded(IDetailPropertyRow& ChildRow) override;
 
@@ -60,7 +59,7 @@ protected:
 	TSharedRef<SWidget> OnPropertyNameContent(TSharedPtr<IPropertyHandle> ChildPropertyHandle, TSharedPtr<SInlineEditableTextBlock> InlineWidget) const;
 
 	TSharedPtr<IPropertyHandle> BagStructProperty;
-	IPropertyUtilities* PropUtils = nullptr;
+	TSharedPtr<IPropertyUtilities> PropUtils;
 	bool bFixedLayout = false;
 };
 
diff --git a/Engine/Plugins/Runtime/StateTree/Source/StateTreeEditorModule/Private/Customizations/StateTreeEditorDataDetails.cpp b/Engine/Plugins/Runtime/StateTree/Source/StateTreeEditorModule/Private/Customizations/StateTreeEditorDataDetails.cpp
--- a/Engine/Plugins/Runtime/StateTree/Source/StateTreeEditorModule/Private/Customizations/StateTreeEditorDataDetails.cpp
+++ b/Engine/Plugins/Runtime/StateTree/Source/StateTreeEditorModule/Private/Customizations/StateTreeEditorDataDetails.cpp
@@ -142,7 +142,7 @@ void FStateTreeEditorDataDetails::CustomizeDetails(IDetailLayoutBuilder& DetailB
 	ParametersCategory.SetSortOrder(2);
 	{
 		// Show parameters as a category.
-		IPropertyUtilities* PropUtils = &DetailBuilder.GetPropertyUtilities().Get();
+		TSharedPtr<IPropertyUtilities> PropUtils = DetailBuilder.GetPropertyUtilities();
 		TSharedPtr<IPropertyHandle> RootParametersProperty = DetailBuilder.GetProperty(TEXT("RootParameters")); // FStateTreeStateParameters
 		check(RootParametersProperty);
 		RootParametersProperty->MarkHiddenByCustomization();
diff --git a/Engine/Plugins/Runtime/StateTree/Source/StateTreeEditorModule/Private/Customizations/StateTreeStateParametersDetails.cpp b/Engine/Plugins/Runtime/StateTree/Source/StateTreeEditorModule/Private/Customizations/StateTreeStateParametersDetails.cpp
--- a/Engine/Plugins/Runtime/StateTree/Source/StateTreeEditorModule/Private/Customizations/StateTreeStateParametersDetails.cpp
+++ b/Engine/Plugins/Runtime/StateTree/Source/StateTreeEditorModule/Private/Customizations/StateTreeStateParametersDetails.cpp
@@ -24,7 +24,7 @@
 class FStateTreeStateParametersInstanceDataDetails : public FPropertyBagInstanceDataDetails
 {
 public:
-	FStateTreeStateParametersInstanceDataDetails(TSharedPtr<IPropertyHandle> InStructProperty, IPropertyUtilities* InPropUtils, const bool bInFixedLayout, FGuid InID, UStateTreeEditorData* InEditorData)
+	FStateTreeStateParametersInstanceDataDetails(TSharedPtr<IPropertyHandle> InStructProperty, TSharedPtr<IPropertyUtilities> InPropUtils, const bool bInFixedLayout, FGuid InID, UStateTreeEditorData* InEditorData)
 		: FPropertyBagInstanceDataDetails(InStructProperty, InPropUtils, bInFixedLayout)
 		, EditorData(InEditorData)
 	{
@@ -99,7 +99,7 @@ TSharedRef<IPropertyTypeCustomization> FStateTreeStateParametersDetails::MakeIns
 void FStateTreeStateParametersDetails::CustomizeHeader(TSharedRef<class IPropertyHandle> StructPropertyHandle, class FDetailWidgetRow& HeaderRow, IPropertyTypeCustomizationUtils& StructCustomizationUtils)
 {
 	StructProperty = StructPropertyHandle;
-	PropUtils = StructCustomizationUtils.GetPropertyUtilities().Get();
+	PropUtils = StructCustomizationUtils.GetPropertyUtilities();
 
 	ParametersProperty = StructProperty->GetChildHandle(TEXT("Parameters"));
 	FixedLayoutProperty = StructProperty->GetChildHandle(TEXT("bFixedLayout"));
diff --git a/Engine/Plugins/Runtime/StateTree/Source/StateTreeEditorModule/Private/Customizations/StateTreeStateParametersDetails.h b/Engine/Plugins/Runtime/StateTree/Source/StateTreeEditorModule/Private/Customizations/StateTreeStateParametersDetails.h
--- a/Engine/Plugins/Runtime/StateTree/Source/StateTreeEditorModule/Private/Customizations/StateTreeStateParametersDetails.h
+++ b/Engine/Plugins/Runtime/StateTree/Source/StateTreeEditorModule/Private/Customizations/StateTreeStateParametersDetails.h
@@ -5,6 +5,7 @@
 #include "IPropertyTypeCustomization.h"
 
 class IPropertyHandle;
+class IPropertyUtilities;
 class UStateTree;
 class UStateTreeEditorData;
 
@@ -26,7 +27,7 @@ private:
 
 	void FindOuterObjects();
 
-	class IPropertyUtilities* PropUtils = nullptr;
+	TSharedPtr<IPropertyUtilities> PropUtils;
 
 	TSharedPtr<IPropertyHandle> ParametersProperty;
 	TSharedPtr<IPropertyHandle> FixedLayoutProperty;
