Struct Utils: Remove const from mutable getters

diff --git a/Engine/Plugins/Experimental/StructUtils/Source/StructUtils/Private/InstancedStruct.cpp b/Engine/Plugins/Experimental/StructUtils/Source/StructUtils/Private/InstancedStruct.cpp
--- a/Engine/Plugins/Experimental/StructUtils/Source/StructUtils/Private/InstancedStruct.cpp
+++ b/Engine/Plugins/Experimental/StructUtils/Source/StructUtils/Private/InstancedStruct.cpp
@@ -53,7 +53,7 @@ FInstancedStruct::FInstancedStruct(const FConstStructView InOther)
 
 FInstancedStruct& FInstancedStruct::operator=(const FConstStructView InOther)
 {
-	if(*this != InOther)
+	if (FConstStructView(*this) != InOther)
 	{
 		InitializeAs(InOther.GetScriptStruct(), InOther.GetMemory());
 	}
@@ -104,7 +104,11 @@ void FInstancedStruct::Reset()
 {
 	if (uint8* Memory = GetMutableMemory())
 	{
-		DestroyScriptStruct();
+		check(StructMemory != nullptr);
+		if (ScriptStruct != nullptr)
+		{
+			ScriptStruct->DestroyStruct(GetMutableMemory());
+		}
 		FMemory::Free(Memory);
 	}
 	ResetStructData();
@@ -267,6 +271,8 @@ bool FInstancedStruct::ImportTextItem(const TCHAR*& Buffer, int32 PortFlags, UOb
 	}
 	else
 	{
+		// Make sure the struct is actually loaded before trying to import the text (this boils down to FindObject if the struct is already loaded).
+		// This is needed for user defined structs, BP pin values, config, copy/paste, where there's no guarantee that the referenced struct has actually been loaded yet.
 		UScriptStruct* StructTypePtr = LoadObject<UScriptStruct>(nullptr, StructPathName.ToString());
 		if (!StructTypePtr)
 		{
diff --git a/Engine/Plugins/Experimental/StructUtils/Source/StructUtils/Private/PropertyBag.cpp b/Engine/Plugins/Experimental/StructUtils/Source/StructUtils/Private/PropertyBag.cpp
--- a/Engine/Plugins/Experimental/StructUtils/Source/StructUtils/Private/PropertyBag.cpp
+++ b/Engine/Plugins/Experimental/StructUtils/Source/StructUtils/Private/PropertyBag.cpp
@@ -910,7 +910,7 @@ namespace UE::StructUtils::Private
 		return EPropertyBagResult::Success;
 	}
 
-	void CopyMatchingValuesByID(const FConstStructView Source, const FStructView Target)
+	void CopyMatchingValuesByID(const FConstStructView Source, FStructView Target)
 	{
 		if (!Source.IsValid() || !Target.IsValid())
 		{
@@ -1146,7 +1146,7 @@ void FInstancedPropertyBag::InitializeFromBagStruct(const UPropertyBag* NewBagSt
 	Value.InitializeAs(NewBagStruct);
 }
 
-void FInstancedPropertyBag::CopyMatchingValuesByID(const FInstancedPropertyBag& Other) const
+void FInstancedPropertyBag::CopyMatchingValuesByID(const FInstancedPropertyBag& Other)
 {
 	UE::StructUtils::Private::CopyMatchingValuesByID(Other.Value, Value);
 }
@@ -1265,7 +1265,16 @@ const FPropertyBagPropertyDesc* FInstancedPropertyBag::FindPropertyDescByName(co
 	return nullptr;
 }
 
-void* FInstancedPropertyBag::GetValueAddress(const FPropertyBagPropertyDesc* Desc) const
+const void* FInstancedPropertyBag::GetValueAddress(const FPropertyBagPropertyDesc* Desc) const
+{
+	if (Desc == nullptr || !Value.IsValid())
+	{
+		return nullptr;
+	}
+	return Value.GetMemory() + Desc->CachedProperty->GetOffset_ForInternal();
+}
+
+void* FInstancedPropertyBag::GetMutableValueAddress(const FPropertyBagPropertyDesc* Desc)
 {
 	if (Desc == nullptr || !Value.IsValid())
 	{
@@ -1436,85 +1445,85 @@ TValueOrError<UClass*, EPropertyBagResult> FInstancedPropertyBag::GetValueClass(
 }
 
 
-EPropertyBagResult FInstancedPropertyBag::SetValueBool(const FName Name, const bool bInValue) const
+EPropertyBagResult FInstancedPropertyBag::SetValueBool(const FName Name, const bool bInValue)
 {
 	const FPropertyBagPropertyDesc* Desc = FindPropertyDescByName(Name);
-	return UE::StructUtils::Private::SetPropertyFromInt64(Desc, GetValueAddress(Desc), bInValue ? 1 : 0);
+	return UE::StructUtils::Private::SetPropertyFromInt64(Desc, GetMutableValueAddress(Desc), bInValue ? 1 : 0);
 }
 
-EPropertyBagResult FInstancedPropertyBag::SetValueByte(const FName Name, const uint8 InValue) const
+EPropertyBagResult FInstancedPropertyBag::SetValueByte(const FName Name, const uint8 InValue)
 {
 	const FPropertyBagPropertyDesc* Desc = FindPropertyDescByName(Name);
-	return UE::StructUtils::Private::SetPropertyFromInt64(Desc, GetValueAddress(Desc), InValue);
+	return UE::StructUtils::Private::SetPropertyFromInt64(Desc, GetMutableValueAddress(Desc), InValue);
 }
 
-EPropertyBagResult FInstancedPropertyBag::SetValueInt32(const FName Name, const int32 InValue) const
+EPropertyBagResult FInstancedPropertyBag::SetValueInt32(const FName Name, const int32 InValue)
 {
 	const FPropertyBagPropertyDesc* Desc = FindPropertyDescByName(Name);
-	return UE::StructUtils::Private::SetPropertyFromInt64(Desc, GetValueAddress(Desc), InValue);
+	return UE::StructUtils::Private::SetPropertyFromInt64(Desc, GetMutableValueAddress(Desc), InValue);
 }
 
-EPropertyBagResult FInstancedPropertyBag::SetValueInt64(const FName Name, const int64 InValue) const
+EPropertyBagResult FInstancedPropertyBag::SetValueInt64(const FName Name, const int64 InValue)
 {
 	const FPropertyBagPropertyDesc* Desc = FindPropertyDescByName(Name);
-	return UE::StructUtils::Private::SetPropertyFromInt64(Desc, GetValueAddress(Desc), InValue);
+	return UE::StructUtils::Private::SetPropertyFromInt64(Desc, GetMutableValueAddress(Desc), InValue);
 }
 
-EPropertyBagResult FInstancedPropertyBag::SetValueFloat(const FName Name, const float InValue) const
+EPropertyBagResult FInstancedPropertyBag::SetValueFloat(const FName Name, const float InValue)
 {
 	const FPropertyBagPropertyDesc* Desc = FindPropertyDescByName(Name);
-	return UE::StructUtils::Private::SetPropertyFromDouble(Desc, GetValueAddress(Desc), InValue);
+	return UE::StructUtils::Private::SetPropertyFromDouble(Desc, GetMutableValueAddress(Desc), InValue);
 }
 
-EPropertyBagResult FInstancedPropertyBag::SetValueDouble(const FName Name, const double InValue) const
+EPropertyBagResult FInstancedPropertyBag::SetValueDouble(const FName Name, const double InValue)
 {
 	const FPropertyBagPropertyDesc* Desc = FindPropertyDescByName(Name);
-	return UE::StructUtils::Private::SetPropertyFromDouble(Desc, GetValueAddress(Desc), InValue);
+	return UE::StructUtils::Private::SetPropertyFromDouble(Desc, GetMutableValueAddress(Desc), InValue);
 }
 
-EPropertyBagResult FInstancedPropertyBag::SetValueName(const FName Name, const FName InValue) const
+EPropertyBagResult FInstancedPropertyBag::SetValueName(const FName Name, const FName InValue)
 {
 	const FPropertyBagPropertyDesc* Desc = FindPropertyDescByName(Name);
-	return UE::StructUtils::Private::SetPropertyValue<FName, FNameProperty>(Desc, GetValueAddress(Desc), InValue);
+	return UE::StructUtils::Private::SetPropertyValue<FName, FNameProperty>(Desc, GetMutableValueAddress(Desc), InValue);
 }
 
-EPropertyBagResult FInstancedPropertyBag::SetValueString(const FName Name, const FString& InValue) const
+EPropertyBagResult FInstancedPropertyBag::SetValueString(const FName Name, const FString& InValue)
 {
 	const FPropertyBagPropertyDesc* Desc = FindPropertyDescByName(Name);
-	return UE::StructUtils::Private::SetPropertyValue<FString, FStrProperty>(Desc, GetValueAddress(Desc), InValue);
+	return UE::StructUtils::Private::SetPropertyValue<FString, FStrProperty>(Desc, GetMutableValueAddress(Desc), InValue);
 }
 
-EPropertyBagResult FInstancedPropertyBag::SetValueText(const FName Name, const FText& InValue) const
+EPropertyBagResult FInstancedPropertyBag::SetValueText(const FName Name, const FText& InValue)
 {
 	const FPropertyBagPropertyDesc* Desc = FindPropertyDescByName(Name);
-	return UE::StructUtils::Private::SetPropertyValue<FText, FTextProperty>(Desc, GetValueAddress(Desc), InValue);
+	return UE::StructUtils::Private::SetPropertyValue<FText, FTextProperty>(Desc, GetMutableValueAddress(Desc), InValue);
 }
 
-EPropertyBagResult FInstancedPropertyBag::SetValueEnum(const FName Name, const uint8 InValue, const UEnum* Enum) const
+EPropertyBagResult FInstancedPropertyBag::SetValueEnum(const FName Name, const uint8 InValue, const UEnum* Enum)
 {
 	const FPropertyBagPropertyDesc* Desc = FindPropertyDescByName(Name);
-	return UE::StructUtils::Private::SetPropertyValueAsEnum(Desc, GetValueAddress(Desc), InValue, Enum);
+	return UE::StructUtils::Private::SetPropertyValueAsEnum(Desc, GetMutableValueAddress(Desc), InValue, Enum);
 }
 
-EPropertyBagResult FInstancedPropertyBag::SetValueStruct(const FName Name, FConstStructView InValue) const
+EPropertyBagResult FInstancedPropertyBag::SetValueStruct(const FName Name, FConstStructView InValue)
 {
 	const FPropertyBagPropertyDesc* Desc = FindPropertyDescByName(Name);
-	return UE::StructUtils::Private::SetPropertyValueAsStruct(Desc, GetValueAddress(Desc), InValue);
+	return UE::StructUtils::Private::SetPropertyValueAsStruct(Desc, GetMutableValueAddress(Desc), InValue);
 }
 
-EPropertyBagResult FInstancedPropertyBag::SetValueObject(const FName Name, UObject* InValue) const
+EPropertyBagResult FInstancedPropertyBag::SetValueObject(const FName Name, UObject* InValue)
 {
 	const FPropertyBagPropertyDesc* Desc = FindPropertyDescByName(Name);
-	return UE::StructUtils::Private::SetPropertyValueAsObject(Desc, GetValueAddress(Desc), InValue);
+	return UE::StructUtils::Private::SetPropertyValueAsObject(Desc, GetMutableValueAddress(Desc), InValue);
 }
 
-EPropertyBagResult FInstancedPropertyBag::SetValueClass(const FName Name, UClass* InValue) const
+EPropertyBagResult FInstancedPropertyBag::SetValueClass(const FName Name, UClass* InValue)
 {
 	const FPropertyBagPropertyDesc* Desc = FindPropertyDescByName(Name);
-	return UE::StructUtils::Private::SetPropertyValueAsObject(Desc, GetValueAddress(Desc), InValue);
+	return UE::StructUtils::Private::SetPropertyValueAsObject(Desc, GetMutableValueAddress(Desc), InValue);
 }
 
-EPropertyBagResult FInstancedPropertyBag::SetValue(const FName Name, const FProperty* InSourceProperty, const void* InSourceContainerAddress) const
+EPropertyBagResult FInstancedPropertyBag::SetValue(const FName Name, const FProperty* InSourceProperty, const void* InSourceContainerAddress)
 {
 	const FPropertyBagPropertyDesc* Desc = FindPropertyDescByName(Name);
 	if (Desc == nullptr || InSourceProperty == nullptr || InSourceContainerAddress == nullptr)
@@ -1538,7 +1547,30 @@ EPropertyBagResult FInstancedPropertyBag::SetValue(const FName Name, const FProp
 	return EPropertyBagResult::Success;
 }
 
-TValueOrError<FPropertyBagArrayRef, EPropertyBagResult> FInstancedPropertyBag::GetArrayRef(const FName Name) const
+TValueOrError<FPropertyBagArrayRef, EPropertyBagResult> FInstancedPropertyBag::GetMutableArrayRef(const FName Name)
+{
+	const FPropertyBagPropertyDesc* Desc = FindPropertyDescByName(Name);
+	if (Desc == nullptr)
+	{
+		return MakeError(EPropertyBagResult::PropertyNotFound);
+	}
+	check(Desc->CachedProperty);
+
+	if (Desc->ContainerType != EPropertyBagContainerType::Array)
+	{
+		return MakeError(EPropertyBagResult::TypeMismatch);
+	}
+
+	const void* Address = GetValueAddress(Desc);
+	if (Address == nullptr)
+	{
+		return MakeError(EPropertyBagResult::PropertyNotFound);
+	}
+	
+	return MakeValue(FPropertyBagArrayRef(*Desc, Address));
+}
+
+TValueOrError<const FPropertyBagArrayRef, EPropertyBagResult> FInstancedPropertyBag::GetArrayRef(const FName Name) const
 {
 	const FPropertyBagPropertyDesc* Desc = FindPropertyDescByName(Name);
 	if (Desc == nullptr)
@@ -1797,67 +1829,67 @@ TValueOrError<UClass*, EPropertyBagResult> FPropertyBagArrayRef::GetValueClass(c
 }
 
 
-EPropertyBagResult FPropertyBagArrayRef::SetValueBool(const int32 Index, const bool bInValue) const
+EPropertyBagResult FPropertyBagArrayRef::SetValueBool(const int32 Index, const bool bInValue)
 {
 	return UE::StructUtils::Private::SetPropertyFromInt64(&ValueDesc, GetMutableAddress(Index), bInValue ? 1 : 0);
 }
 
-EPropertyBagResult FPropertyBagArrayRef::SetValueByte(const int32 Index, const uint8 InValue) const
+EPropertyBagResult FPropertyBagArrayRef::SetValueByte(const int32 Index, const uint8 InValue)
 {
 	return UE::StructUtils::Private::SetPropertyFromInt64(&ValueDesc, GetMutableAddress(Index), InValue);
 }
 
-EPropertyBagResult FPropertyBagArrayRef::SetValueInt32(const int32 Index, const int32 InValue) const
+EPropertyBagResult FPropertyBagArrayRef::SetValueInt32(const int32 Index, const int32 InValue)
 {
 	return UE::StructUtils::Private::SetPropertyFromInt64(&ValueDesc, GetMutableAddress(Index), InValue);
 }
 
-EPropertyBagResult FPropertyBagArrayRef::SetValueInt64(const int32 Index, const int64 InValue) const
+EPropertyBagResult FPropertyBagArrayRef::SetValueInt64(const int32 Index, const int64 InValue)
 {
 	return UE::StructUtils::Private::SetPropertyFromInt64(&ValueDesc, GetMutableAddress(Index), InValue);
 }
 
-EPropertyBagResult FPropertyBagArrayRef::SetValueFloat(const int32 Index, const float InValue) const
+EPropertyBagResult FPropertyBagArrayRef::SetValueFloat(const int32 Index, const float InValue)
 {
 	return UE::StructUtils::Private::SetPropertyFromDouble(&ValueDesc, GetMutableAddress(Index), InValue);
 }
 
-EPropertyBagResult FPropertyBagArrayRef::SetValueDouble(const int32 Index, const double InValue) const
+EPropertyBagResult FPropertyBagArrayRef::SetValueDouble(const int32 Index, const double InValue)
 {
 	return UE::StructUtils::Private::SetPropertyFromDouble(&ValueDesc, GetMutableAddress(Index), InValue);
 }
 
-EPropertyBagResult FPropertyBagArrayRef::SetValueName(const int32 Index, const FName InValue) const
+EPropertyBagResult FPropertyBagArrayRef::SetValueName(const int32 Index, const FName InValue)
 {
 	return UE::StructUtils::Private::SetPropertyValue<FName, FNameProperty>(&ValueDesc, GetMutableAddress(Index), InValue);
 }
 
-EPropertyBagResult FPropertyBagArrayRef::SetValueString(const int32 Index, const FString& InValue) const
+EPropertyBagResult FPropertyBagArrayRef::SetValueString(const int32 Index, const FString& InValue)
 {
 	return UE::StructUtils::Private::SetPropertyValue<FString, FStrProperty>(&ValueDesc, GetMutableAddress(Index), InValue);
 }
 
-EPropertyBagResult FPropertyBagArrayRef::SetValueText(const int32 Index, const FText& InValue) const
+EPropertyBagResult FPropertyBagArrayRef::SetValueText(const int32 Index, const FText& InValue)
 {
 	return UE::StructUtils::Private::SetPropertyValue<FText, FTextProperty>(&ValueDesc, GetMutableAddress(Index), InValue);
 }
 
-EPropertyBagResult FPropertyBagArrayRef::SetValueEnum(const int32 Index, const uint8 InValue, const UEnum* Enum) const
+EPropertyBagResult FPropertyBagArrayRef::SetValueEnum(const int32 Index, const uint8 InValue, const UEnum* Enum)
 {
 	return UE::StructUtils::Private::SetPropertyValueAsEnum(&ValueDesc, GetMutableAddress(Index), InValue, Enum);
 }
 
-EPropertyBagResult FPropertyBagArrayRef::SetValueStruct(const int32 Index, FConstStructView InValue) const
+EPropertyBagResult FPropertyBagArrayRef::SetValueStruct(const int32 Index, FConstStructView InValue)
 {
 	return UE::StructUtils::Private::SetPropertyValueAsStruct(&ValueDesc, GetMutableAddress(Index), InValue);
 }
 
-EPropertyBagResult FPropertyBagArrayRef::SetValueObject(const int32 Index, UObject* InValue) const
+EPropertyBagResult FPropertyBagArrayRef::SetValueObject(const int32 Index, UObject* InValue)
 {
 	return UE::StructUtils::Private::SetPropertyValueAsObject(&ValueDesc, GetMutableAddress(Index), InValue);
 }
 
-EPropertyBagResult FPropertyBagArrayRef::SetValueClass(const int32 Index, UClass* InValue) const
+EPropertyBagResult FPropertyBagArrayRef::SetValueClass(const int32 Index, UClass* InValue)
 {
 	return UE::StructUtils::Private::SetPropertyValueAsObject(&ValueDesc, GetMutableAddress(Index), InValue);
 }
diff --git a/Engine/Plugins/Experimental/StructUtils/Source/StructUtils/Public/InstancedStruct.h b/Engine/Plugins/Experimental/StructUtils/Source/StructUtils/Public/InstancedStruct.h
--- a/Engine/Plugins/Experimental/StructUtils/Source/StructUtils/Public/InstancedStruct.h
+++ b/Engine/Plugins/Experimental/StructUtils/Source/StructUtils/Public/InstancedStruct.h
@@ -29,7 +29,11 @@ public:
 
 	explicit FInstancedStruct(const UScriptStruct* InScriptStruct);
 
-	FInstancedStruct(const FConstStructView InOther);
+	/**
+	 * This constructor is explicit to avoid accidentally converting struct views to instanced structs (which would result in costly copy of the struct to be made).
+	 * Implicit conversion could happen e.g. when comparing FInstancedStruct to FConstStructView.
+	 */
+	explicit FInstancedStruct(const FConstStructView InOther);
 
 	FInstancedStruct(const FInstancedStruct& InOther)
 	{
@@ -188,14 +192,14 @@ public:
 	}
 
 	/** Returns a mutable pointer to struct memory. */
-	uint8* GetMutableMemory() const
+	uint8* GetMutableMemory()
 	{
 		return StructMemory;
 	}
 
 	/** Returns mutable reference to the struct, this getter assumes that all data is valid. */
 	template<typename T>
-	T& GetMutable() const
+	T& GetMutable()
 	{
 		uint8* Memory = GetMutableMemory();
 		const UScriptStruct* Struct = GetScriptStruct();
@@ -207,7 +211,7 @@ public:
 
 	/** Returns mutable pointer to the struct, or nullptr if cast is not valid. */
 	template<typename T>
-	T* GetMutablePtr() const
+	T* GetMutablePtr()
 	{
 		uint8* Memory = GetMutableMemory();
 		const UScriptStruct* Struct = GetScriptStruct();
@@ -235,34 +239,8 @@ public:
 		return !Identical(&Other, PPF_None);
 	}
 
-	/** Comparison operators. Note: it does not compare the internal structure itself*/
-	template <typename OtherType>
-	bool operator==(const OtherType& Other) const
-	{
-		if ((GetScriptStruct() != Other.GetScriptStruct()) || (GetMemory() != Other.GetMemory()))
-		{
-			return false;
-		}
-		return true;
-	}
-
-	template <typename OtherType>
-	bool operator!=(const OtherType& Other) const
-	{
-		return !operator==(Other);
-	}
-
 protected:
 
-	void DestroyScriptStruct() const
-	{
-		check(StructMemory != nullptr);
-		if (ScriptStruct != nullptr)
-		{
-			ScriptStruct->DestroyStruct(GetMutableMemory());
-		}
-	}
-
 	FInstancedStruct(const UScriptStruct* InScriptStruct, uint8* InStructMemory)
 		: ScriptStruct(InScriptStruct)
 		, StructMemory(InStructMemory)
diff --git a/Engine/Plugins/Experimental/StructUtils/Source/StructUtils/Public/PropertyBag.h b/Engine/Plugins/Experimental/StructUtils/Source/StructUtils/Public/PropertyBag.h
--- a/Engine/Plugins/Experimental/StructUtils/Source/StructUtils/Public/PropertyBag.h
+++ b/Engine/Plugins/Experimental/StructUtils/Source/StructUtils/Public/PropertyBag.h
@@ -240,7 +240,7 @@ struct STRUCTUTILS_API FInstancedPropertyBag
 	 * The properties are matched between the bags based on the property ID.
 	 * @param Other Reference to the bag to copy the values from
 	 */
-	void CopyMatchingValuesByID(const FInstancedPropertyBag& NewDescs) const;
+	void CopyMatchingValuesByID(const FInstancedPropertyBag& NewDescs);
 
 	/** Returns number of the Properties in this Property Bag */
 	int32 GetNumPropertiesInBag() const;
@@ -313,7 +313,7 @@ struct STRUCTUTILS_API FInstancedPropertyBag
 	FConstStructView GetValue() const { return Value; };
 
 	/** @return const view to the struct that holds the values. NOTE: The returned value/view cannot be serialized, use this to access the struct only temporarily. */
-	FStructView GetMutableValue() const { return Value; };
+	FStructView GetMutableValue() { return Value; };
 	
 	/**
 	 * Getters
@@ -390,23 +390,23 @@ struct STRUCTUTILS_API FInstancedPropertyBag
 	 * Value Setters. A property must exists in that bag before it can be set.  
 	 * Numeric types (bool, int32, int64, float, double) support type conversion.
 	 */
-	EPropertyBagResult SetValueBool(const FName Name, const bool bInValue) const;
-	EPropertyBagResult SetValueByte(const FName Name, const uint8 InValue) const;
-	EPropertyBagResult SetValueInt32(const FName Name, const int32 InValue) const;
-	EPropertyBagResult SetValueInt64(const FName Name, const int64 InValue) const;
-	EPropertyBagResult SetValueFloat(const FName Name, const float InValue) const;
-	EPropertyBagResult SetValueDouble(const FName Name, const double InValue) const;
-	EPropertyBagResult SetValueName(const FName Name, const FName InValue) const;
-	EPropertyBagResult SetValueString(const FName Name, const FString& InValue) const;
-	EPropertyBagResult SetValueText(const FName Name, const FText& InValue) const;
-	EPropertyBagResult SetValueEnum(const FName Name, const uint8 InValue, const UEnum* Enum) const;
-	EPropertyBagResult SetValueStruct(const FName Name, FConstStructView InValue) const;
-	EPropertyBagResult SetValueObject(const FName Name, UObject* InValue) const;
-	EPropertyBagResult SetValueClass(const FName Name, UClass* InValue) const;
+	EPropertyBagResult SetValueBool(const FName Name, const bool bInValue);
+	EPropertyBagResult SetValueByte(const FName Name, const uint8 InValue);
+	EPropertyBagResult SetValueInt32(const FName Name, const int32 InValue);
+	EPropertyBagResult SetValueInt64(const FName Name, const int64 InValue);
+	EPropertyBagResult SetValueFloat(const FName Name, const float InValue);
+	EPropertyBagResult SetValueDouble(const FName Name, const double InValue);
+	EPropertyBagResult SetValueName(const FName Name, const FName InValue);
+	EPropertyBagResult SetValueString(const FName Name, const FString& InValue);
+	EPropertyBagResult SetValueText(const FName Name, const FText& InValue);
+	EPropertyBagResult SetValueEnum(const FName Name, const uint8 InValue, const UEnum* Enum);
+	EPropertyBagResult SetValueStruct(const FName Name, FConstStructView InValue);
+	EPropertyBagResult SetValueObject(const FName Name, UObject* InValue);
+	EPropertyBagResult SetValueClass(const FName Name, UClass* InValue);
 
 	/** Sets enum value specified type. */
 	template <typename T>
-	EPropertyBagResult SetValueEnum(const FName Name, const T InValue) const
+	EPropertyBagResult SetValueEnum(const FName Name, const T InValue)
 	{
 		static_assert(TIsEnum<T>::Value, "Should only call this with enum types");
 		return SetValueEnum(Name, (uint8)InValue, StaticEnum<T>());
@@ -414,14 +414,14 @@ struct STRUCTUTILS_API FInstancedPropertyBag
 
 	/** Sets struct value specified type. */
 	template <typename T>
-	EPropertyBagResult SetValueStruct(const FName Name, const T& InValue) const
+	EPropertyBagResult SetValueStruct(const FName Name, const T& InValue)
 	{
 		return SetValueStruct(Name, FConstStructView::Make(InValue));
 	}
 
 	/** Sets object pointer value specified type. */
 	template <typename T>
-	EPropertyBagResult SetValueObject(const FName Name, T* InValue) const
+	EPropertyBagResult SetValueObject(const FName Name, T* InValue)
 	{
 		static_assert(TIsDerivedFrom<T, UObject>::Value, "Should only call this with object types");
 		return SetValueObject(Name, (UObject*)InValue);
@@ -431,19 +431,27 @@ struct STRUCTUTILS_API FInstancedPropertyBag
 	 * Sets property value from given source property and source container address
 	 * A property must exists in that bag before it can be set. 
 	 */
-	EPropertyBagResult SetValue(const FName Name, const FProperty* InSourceProperty, const void* InSourceContainerAddress) const;
+	EPropertyBagResult SetValue(const FName Name, const FProperty* InSourceProperty, const void* InSourceContainerAddress);
 
 	/**
 	 * Returns helper class to modify and access an array property.
 	 * Note: Note: The array reference is not valid after the layout of the referenced property bag has changed!
 	 * @returns helper class to modify and access arrays
 	*/
-	TValueOrError<FPropertyBagArrayRef, EPropertyBagResult> GetArrayRef(const FName Name) const;
+	TValueOrError<FPropertyBagArrayRef, EPropertyBagResult> GetMutableArrayRef(const FName Name);
+
+	/**
+	 * Returns helper class to modify and access an array property.
+	 * Note: Note: The array reference is not valid after the layout of the referenced property bag has changed!
+	 * @returns helper class to modify and access arrays
+	*/
+	TValueOrError<const FPropertyBagArrayRef, EPropertyBagResult> GetArrayRef(const FName Name) const;
 
 	bool Serialize(FArchive& Ar);
 
 protected:
-	void* GetValueAddress(const FPropertyBagPropertyDesc* Desc) const;
+	const void* GetValueAddress(const FPropertyBagPropertyDesc* Desc) const;
+	void* GetMutableValueAddress(const FPropertyBagPropertyDesc* Desc);
 	
 	UPROPERTY(EditAnywhere, Category="")
 	FInstancedStruct Value;
@@ -591,23 +599,23 @@ public:
 	 * Value Setters. A property must exists in that bag before it can be set.  
 	 * Numeric types (bool, int32, int64, float, double) support type conversion.
 	 */
-	EPropertyBagResult SetValueBool(const int32 Index, const bool bInValue) const;
-	EPropertyBagResult SetValueByte(const int32 Index, const uint8 InValue) const;
-	EPropertyBagResult SetValueInt32(const int32 Index, const int32 InValue) const;
-	EPropertyBagResult SetValueInt64(const int32 Index, const int64 InValue) const;
-	EPropertyBagResult SetValueFloat(const int32 Index, const float InValue) const;
-	EPropertyBagResult SetValueDouble(const int32 Index, const double InValue) const;
-	EPropertyBagResult SetValueName(const int32 Index, const FName InValue) const;
-	EPropertyBagResult SetValueString(const int32 Index, const FString& InValue) const;
-	EPropertyBagResult SetValueText(const int32 Index, const FText& InValue) const;
-	EPropertyBagResult SetValueEnum(const int32 Index, const uint8 InValue, const UEnum* Enum) const;
-	EPropertyBagResult SetValueStruct(const int32 Index, FConstStructView InValue) const;
-	EPropertyBagResult SetValueObject(const int32 Index, UObject* InValue) const;
-	EPropertyBagResult SetValueClass(const int32 Index, UClass* InValue) const;
+	EPropertyBagResult SetValueBool(const int32 Index, const bool bInValue);
+	EPropertyBagResult SetValueByte(const int32 Index, const uint8 InValue);
+	EPropertyBagResult SetValueInt32(const int32 Index, const int32 InValue);
+	EPropertyBagResult SetValueInt64(const int32 Index, const int64 InValue);
+	EPropertyBagResult SetValueFloat(const int32 Index, const float InValue);
+	EPropertyBagResult SetValueDouble(const int32 Index, const double InValue);
+	EPropertyBagResult SetValueName(const int32 Index, const FName InValue);
+	EPropertyBagResult SetValueString(const int32 Index, const FString& InValue);
+	EPropertyBagResult SetValueText(const int32 Index, const FText& InValue);
+	EPropertyBagResult SetValueEnum(const int32 Index, const uint8 InValue, const UEnum* Enum);
+	EPropertyBagResult SetValueStruct(const int32 Index, FConstStructView InValue);
+	EPropertyBagResult SetValueObject(const int32 Index, UObject* InValue);
+	EPropertyBagResult SetValueClass(const int32 Index, UClass* InValue);
 
 	/** Sets enum value specified type. */
 	template <typename T>
-	EPropertyBagResult SetValueEnum(const int32 Index, const T InValue) const
+	EPropertyBagResult SetValueEnum(const int32 Index, const T InValue)
 	{
 		static_assert(TIsEnum<T>::Value, "Should only call this with enum types");
 		return SetValueEnum(Index, (uint8)InValue, StaticEnum<T>());
@@ -615,14 +623,14 @@ public:
 
 	/** Sets struct value specified type. */
 	template <typename T>
-	EPropertyBagResult SetValueStruct(const int32 Index, const T& InValue) const
+	EPropertyBagResult SetValueStruct(const int32 Index, const T& InValue)
 	{
 		return SetValueStruct(Index, FConstStructView::Make(InValue));
 	}
 
 	/** Sets object pointer value specified type. */
 	template <typename T>
-	EPropertyBagResult SetValueObject(const int32 Index, T* InValue) const
+	EPropertyBagResult SetValueObject(const int32 Index, T* InValue)
 	{
 		static_assert(TIsDerivedFrom<T, UObject>::Value, "Should only call this with object types");
 		return SetValueObject(Index, (UObject*)InValue);
diff --git a/Engine/Plugins/Experimental/StructUtils/Source/StructUtils/Public/SharedStruct.h b/Engine/Plugins/Experimental/StructUtils/Source/StructUtils/Public/SharedStruct.h
--- a/Engine/Plugins/Experimental/StructUtils/Source/StructUtils/Public/SharedStruct.h
+++ b/Engine/Plugins/Experimental/StructUtils/Source/StructUtils/Public/SharedStruct.h
@@ -280,7 +280,7 @@ struct STRUCTUTILS_API FSharedStruct : public FConstSharedStruct
 	}
 
 	/** Returns a mutable pointer to struct memory. This const_cast here is safe as a ClassName can only be setup from mutable non const memory. */
-	uint8* GetMutableMemory() const
+	uint8* GetMutableMemory()
 	{
 		const uint8* Memory = GetMemory();
 		return const_cast<uint8*>(Memory);
@@ -288,7 +288,7 @@ struct STRUCTUTILS_API FSharedStruct : public FConstSharedStruct
 
 	/** Returns mutable reference to the struct, this getter assumes that all data is valid. */
 	template<typename T>
-	T& GetMutable() const
+	T& GetMutable()
 	{
 		uint8* Memory = GetMutableMemory();
 		const UScriptStruct* Struct = GetScriptStruct();
@@ -300,7 +300,7 @@ struct STRUCTUTILS_API FSharedStruct : public FConstSharedStruct
 
 	/** Returns mutable pointer to the struct, or nullptr if cast is not valid. */
 	template<typename T>
-	T* GetMutablePtr() const
+	T* GetMutablePtr()
 	{
 		uint8* Memory = GetMutableMemory();
 		const UScriptStruct* Struct = GetScriptStruct();
diff --git a/Engine/Plugins/Experimental/StructUtils/Source/StructUtils/Public/StructView.h b/Engine/Plugins/Experimental/StructUtils/Source/StructUtils/Public/StructView.h
--- a/Engine/Plugins/Experimental/StructUtils/Source/StructUtils/Public/StructView.h
+++ b/Engine/Plugins/Experimental/StructUtils/Source/StructUtils/Public/StructView.h
@@ -160,11 +160,11 @@ public:
 		: FConstStructView(InScriptStruct, InStructMemory)
 	{}
 
-	FStructView(const FInstancedStruct& InstancedStruct)
+	FStructView(FInstancedStruct& InstancedStruct)
 		: FConstStructView(InstancedStruct.GetScriptStruct(), InstancedStruct.GetMutableMemory())
 	{}
 
-	FStructView(const FSharedStruct& SharedStruct)
+	FStructView(FSharedStruct& SharedStruct)
 		: FConstStructView(SharedStruct.GetScriptStruct(), SharedStruct.GetMutableMemory())
 	{}
 
@@ -185,7 +185,7 @@ public:
 	}
 
 	/** Returns a mutable pointer to struct memory. This const_cast here is safe as a ClassName can only be setup from mutable non const memory. */
-	uint8* GetMutableMemory() const
+	uint8* GetMutableMemory()
 	{
 		const uint8* Memory = GetMemory();
 		return const_cast<uint8*>(Memory);
@@ -193,7 +193,7 @@ public:
 
 	/** Returns mutable reference to the struct, this getter assumes that all data is valid. */
 	template<typename T>
-	T& GetMutable() const
+	T& GetMutable()
 	{
 		uint8* Memory = GetMutableMemory();
 		const UScriptStruct* Struct = GetScriptStruct();
@@ -205,7 +205,7 @@ public:
 
 	/** Returns mutable pointer to the struct, or nullptr if cast is not valid. */
 	template<typename T>
-	T* GetMutablePtr() const
+	T* GetMutablePtr()
 	{
 		uint8* Memory = GetMutableMemory();
 		const UScriptStruct* Struct = GetScriptStruct();
diff --git a/Engine/Plugins/Experimental/StructUtils/Source/StructUtilsTestSuite/Private/PropertyBagTest.cpp b/Engine/Plugins/Experimental/StructUtils/Source/StructUtilsTestSuite/Private/PropertyBagTest.cpp
--- a/Engine/Plugins/Experimental/StructUtils/Source/StructUtilsTestSuite/Private/PropertyBagTest.cpp
+++ b/Engine/Plugins/Experimental/StructUtils/Source/StructUtilsTestSuite/Private/PropertyBagTest.cpp
@@ -305,7 +305,7 @@ struct FTest_Arrays : FAITestBase
 			{ FloatArrayName, EPropertyBagContainerType::Array, EPropertyBagPropertyType::Float },
 		});
 
-		TValueOrError<FPropertyBagArrayRef, EPropertyBagResult> FloatArrayRes = Bag.GetArrayRef(FloatArrayName);
+		TValueOrError<const FPropertyBagArrayRef, EPropertyBagResult> FloatArrayRes = Bag.GetArrayRef(FloatArrayName);
 		AITEST_TRUE(TEXT("Get float array should succeed"), FloatArrayRes.IsValid());
 
 		FPropertyBagArrayRef FloatArray = FloatArrayRes.GetValue();
