[StructUtils] moved `FInstancedStruct` registration for localization from constructors to module startup.
This fixes potential issues when the type get used in static variables in other modules.
#rnx

diff --git a/Engine/Plugins/Experimental/StructUtils/Source/StructUtils/Private/InstancedStruct.cpp b/Engine/Plugins/Experimental/StructUtils/Source/StructUtils/Private/InstancedStruct.cpp
--- a/Engine/Plugins/Experimental/StructUtils/Source/StructUtils/Private/InstancedStruct.cpp
+++ b/Engine/Plugins/Experimental/StructUtils/Source/StructUtils/Private/InstancedStruct.cpp
@@ -23,30 +23,25 @@ void GatherForLocalization(const FString& PathToParent, const UScriptStruct* Str
 		PropertyLocalizationDataGatherer.GatherLocalizationDataFromStructWithCallbacks(PathToParent + TEXT(".StructInstance"), StructTypePtr, ThisInstance->GetMemory(), DefaultInstance ? DefaultInstance->GetMemory() : nullptr, GatherTextFlags);
 	}
 }
-#endif // WITH_EDITORONLY_DATA
 
-void RegisterForLocalization()
+void RegisterInstancedStructForLocalization()
 {
-#if WITH_EDITORONLY_DATA
 	{ static const FAutoRegisterLocalizationDataGatheringCallback AutomaticRegistrationOfLocalizationGatherer(TBaseStructure<FInstancedStruct>::Get(), &GatherForLocalization); }
-#endif
 }
+#endif // WITH_EDITORONLY_DATA
 }
 
 FInstancedStruct::FInstancedStruct()
 {
-	UE::StructUtils::Private::RegisterForLocalization();
 }
 
 FInstancedStruct::FInstancedStruct(const UScriptStruct* InScriptStruct)
 {
-	UE::StructUtils::Private::RegisterForLocalization();
 	InitializeAs(InScriptStruct, nullptr);
 }
 
 FInstancedStruct::FInstancedStruct(const FConstStructView InOther)
 {
-	UE::StructUtils::Private::RegisterForLocalization();
 	InitializeAs(InOther.GetScriptStruct(), InOther.GetMemory());
 }
 
diff --git a/Engine/Plugins/Experimental/StructUtils/Source/StructUtils/Private/InstancedStructContainer.cpp b/Engine/Plugins/Experimental/StructUtils/Source/StructUtils/Private/InstancedStructContainer.cpp
--- a/Engine/Plugins/Experimental/StructUtils/Source/StructUtils/Private/InstancedStructContainer.cpp
+++ b/Engine/Plugins/Experimental/StructUtils/Source/StructUtils/Private/InstancedStructContainer.cpp
@@ -3,25 +3,8 @@
 
 #include UE_INLINE_GENERATED_CPP_BY_NAME(InstancedStructContainer)
 
-class FPropertyLocalizationDataGatherer;
-enum class EPropertyLocalizationGathererTextFlags : uint8;
-
-// From InstancedStruct.cpp to support localization.
-namespace UE::StructUtils::Private
-{
-#if WITH_EDITORONLY_DATA
-	void GatherForLocalization(const FString& PathToParent, const UScriptStruct* Struct, const void* StructData, const void* DefaultStructData,
-								FPropertyLocalizationDataGatherer& PropertyLocalizationDataGatherer, const EPropertyLocalizationGathererTextFlags GatherTextFlags);
-#endif // WITH_EDITORONLY_DATA
-
-	void RegisterForLocalization();
-
-} // UE::StructUtils::Private
-
-
 FInstancedStructContainer::FInstancedStructContainer()
 {
-	UE::StructUtils::Private::RegisterForLocalization();
 }
 
 FInstancedStructContainer::FInstancedStructContainer(const FInstancedStructContainer& InOther)
diff --git a/Engine/Plugins/Experimental/StructUtils/Source/StructUtils/Private/StructUtilsModule.cpp b/Engine/Plugins/Experimental/StructUtils/Source/StructUtils/Private/StructUtilsModule.cpp
--- a/Engine/Plugins/Experimental/StructUtils/Source/StructUtils/Private/StructUtilsModule.cpp
+++ b/Engine/Plugins/Experimental/StructUtils/Source/StructUtils/Private/StructUtilsModule.cpp
@@ -2,6 +2,10 @@
 
 #include "StructUtilsModule.h"
 
+#if WITH_EDITORONLY_DATA
+#include "InstancedStruct.h"
+#endif // WITH_EDITORONLY_DATA
+
 #define LOCTEXT_NAMESPACE "StructUtils"
 
 class FStructUtilsModule : public IStructUtilsModule
@@ -18,6 +22,10 @@ void FStructUtilsModule::StartupModule()
 	// Called right after the module DLL has been loaded and the module object has been created
 	// Load dependent modules here, and they will be guaranteed to be available during ShutdownModule. ie:
 	// 		FModuleManager::Get().LoadModuleChecked(TEXT("HTTP"));
+	
+#if WITH_EDITORONLY_DATA
+	UE::StructUtils::Private::RegisterInstancedStructForLocalization();
+#endif // WITH_EDITORONLY_DATA
 }
 
 void FStructUtilsModule::ShutdownModule()
diff --git a/Engine/Plugins/Experimental/StructUtils/Source/StructUtils/Public/InstancedStruct.h b/Engine/Plugins/Experimental/StructUtils/Source/StructUtils/Public/InstancedStruct.h
--- a/Engine/Plugins/Experimental/StructUtils/Source/StructUtils/Public/InstancedStruct.h
+++ b/Engine/Plugins/Experimental/StructUtils/Source/StructUtils/Public/InstancedStruct.h
@@ -275,3 +275,10 @@ struct TStructOpsTypeTraits<FInstancedStruct> : public TStructOpsTypeTraitsBase2
 		WithNetSerializer = true,
 	};
 };
+
+#if WITH_EDITORONLY_DATA
+namespace UE::StructUtils::Private
+{
+void RegisterInstancedStructForLocalization();
+}
+#endif // WITH_EDITORONLY_DATA
\ No newline at end of file
