Added FStructArrayView constructor for procedural creation of instances based on a known type, pointed to the data and number of elements.
Also renamed FStructArrayView.FragmentType to ElementType and made it const.

diff --git a/Engine/Plugins/Experimental/StructUtils/Source/StructUtils/Public/StructArrayView.h b/Engine/Plugins/Experimental/StructUtils/Source/StructUtils/Public/StructArrayView.h
--- a/Engine/Plugins/Experimental/StructUtils/Source/StructUtils/Public/StructArrayView.h
+++ b/Engine/Plugins/Experimental/StructUtils/Source/StructUtils/Public/StructArrayView.h
@@ -15,14 +15,14 @@ struct FStructArrayView
 		: ElementSize(Src.ElementSize)
 		, DataPtr(Src.DataPtr)
 		, NumElements(Src.NumElements)
-		, FragmentType(Src.FragmentType)
+		, ElementType(Src.ElementType)
 	{}
 
 	FStructArrayView(const FStructArrayView& Src)
 		: ElementSize(Src.ElementSize)
 		, DataPtr(Src.DataPtr)
 		, NumElements(Src.NumElements)
-		, FragmentType(Src.FragmentType)
+		, ElementType(Src.ElementType)
 	{}
 	
 	template<typename T>
@@ -30,7 +30,7 @@ struct FStructArrayView
 		: ElementSize(sizeof(T))
 		, DataPtr(InArray.GetData())
 		, NumElements(InArray.Num())
-		, FragmentType(*StaticStruct<typename TRemoveReference<T>::Type>())
+		, ElementType(*StaticStruct<typename TRemoveReference<T>::Type>())
 
 	{}
 
@@ -39,7 +39,7 @@ struct FStructArrayView
 		: ElementSize(sizeof(T))
 		, DataPtr(InArrayView.GetData())
 		, NumElements(InArrayView.Num())
-		, FragmentType(*StaticStruct<typename TRemoveReference<T>::Type>())
+		, ElementType(*StaticStruct<typename TRemoveReference<T>::Type>())
 	{}
 
 	/**
@@ -49,18 +49,34 @@ struct FStructArrayView
 		: ElementSize(Src.ElementSize)
 		, DataPtr((uint8*)Src.DataPtr + Src.ElementSize * OffsetElements)
 		, NumElements(InNumElements)
-		, FragmentType(Src.FragmentType)
+		, ElementType(Src.ElementType)
 	{
 		checkf(InNumElements + OffsetElements <= Src.NumElements, TEXT("Requested range passes over the end of the view, %d + %d> %d")
 			, InNumElements, OffsetElements, Src.NumElements);
 	}
 
+	/** 
+	 * Construct a view with given data.
+	 */
+	FStructArrayView(const UScriptStruct& InType, void* InDataPtr, const int32 InNumElements)
+		: ElementSize(InType.GetStructureSize())
+		, DataPtr(InDataPtr)
+		, NumElements(InNumElements)
+		, ElementType(InType)
+	{
+		check(InNumElements >= 0);
+		check(InDataPtr != nullptr || InNumElements == 0);
+	}
+
 	SIZE_T GetElementSize() const { return ElementSize; }
 	const void* GetData() const { return DataPtr; }
 	const void* GetDataAt(const int32 Index) const { return (uint8*)DataPtr + Index * ElementSize; }
 	void* GetMutableDataAt(const int32 Index) const { return (uint8*)DataPtr + Index * ElementSize; }
 	int32 Num() const { return NumElements; }
-	const UScriptStruct& GetFragmentType() const { return FragmentType; }
+	const UScriptStruct& GetElementType() const { return ElementType; }
+
+	UE_DEPRECATED(5.2, "This function was renamed to GetElementType. Use that name instead.")
+	const UScriptStruct& GetFragmentType() const { return GetElementType(); }
 
 	template<typename T>
 	const T& GetElementAt(const int32 Index) const
@@ -71,7 +87,7 @@ struct FStructArrayView
 	template<typename T>
 	const T& GetElementAtChecked(const int32 Index) const
 	{
-		check(TBaseStructure<T>::Get() == &FragmentType);
+		check(TBaseStructure<T>::Get() == &ElementType);
 		return *((T*)GetDataAt(Index));
 	}
 
@@ -84,7 +100,7 @@ struct FStructArrayView
 	template<typename T>
 	T& GetMutableElementAtChecked(const int32 Index)
 	{
-		check(TBaseStructure<T>::Get() == &FragmentType);
+		check(TBaseStructure<T>::Get() == &ElementType);
 		return *((T*)GetDataAt(Index));
 	}
 
@@ -99,7 +115,7 @@ private:
 	SIZE_T ElementSize = 0;
 	void* DataPtr = nullptr;
 	int32 NumElements = 0;
-	UScriptStruct& FragmentType;
+	const UScriptStruct& ElementType;
 };
 
 #if UE_ENABLE_INCLUDE_ORDER_DEPRECATED_IN_5_2
diff --git a/Engine/Plugins/Runtime/MassEntity/Source/MassEntity/Private/MassArchetypeData.cpp b/Engine/Plugins/Runtime/MassEntity/Source/MassEntity/Private/MassArchetypeData.cpp
--- a/Engine/Plugins/Runtime/MassEntity/Source/MassEntity/Private/MassArchetypeData.cpp
+++ b/Engine/Plugins/Runtime/MassEntity/Source/MassEntity/Private/MassArchetypeData.cpp
@@ -1178,7 +1178,7 @@ void FMassArchetypeData::BatchSetFragmentValues(TConstArrayView<FMassArchetypeEn
 			FStructArrayView FragmentPayload = Payload[i];
 			check(FragmentPayload.Num() - EntitiesHandled >= EntityRange.Length);
 
-			const UScriptStruct& FragmentType = FragmentPayload.GetFragmentType();
+			const UScriptStruct& FragmentType = FragmentPayload.GetElementType();
 
 			const int32 FragmentIndex = FragmentIndexMap.FindChecked(&FragmentType);
 			void* Dst = FragmentConfigs[FragmentIndex].GetFragmentData(Chunk.GetRawMemory(), EntityRange.SubchunkStart);
