Missing include for unity build

diff --git a/Engine/Plugins/Experimental/StructUtils/Source/StructUtils/Public/StructUtils.h b/Engine/Plugins/Experimental/StructUtils/Source/StructUtils/Public/StructUtils.h
--- a/Engine/Plugins/Experimental/StructUtils/Source/StructUtils/Public/StructUtils.h
+++ b/Engine/Plugins/Experimental/StructUtils/Source/StructUtils/Public/StructUtils.h
@@ -3,6 +3,7 @@
 #pragma once
 
 #include "Templates/UnrealTypeTraits.h"
+#include "UObject/Class.h"
 
 namespace UE::StructUtils
 {
