Added a missing #include and a missing deprecated function.

diff --git a/Engine/Plugins/Experimental/StructUtils/Source/StructUtils/Public/StructView.h b/Engine/Plugins/Experimental/StructUtils/Source/StructUtils/Public/StructView.h
--- a/Engine/Plugins/Experimental/StructUtils/Source/StructUtils/Public/StructView.h
+++ b/Engine/Plugins/Experimental/StructUtils/Source/StructUtils/Public/StructView.h
@@ -4,6 +4,7 @@
 
 #include "InstancedStruct.h"
 #include "SharedStruct.h"
+#include "StructUtils.h"
 
 ///////////////////////////////////////////////////////////////// FStructView /////////////////////////////////////////////////////////////////
 
@@ -87,6 +88,13 @@ public:
 		return StructMemory;
 	}
 
+	/** Returns a mutable pointer to struct memory. */
+	UE_DEPRECATED(5.3, "Use GetMemory() instead")
+	uint8* GetMutableMemory()
+	{
+		return GetMemory();
+	}
+
 	/** Reset to empty. */
 	void Reset()
 	{
