Very minor change operator== had not been updated correcty when refactoring the code and was != when should have been ==

diff --git a/Engine/Plugins/Experimental/StructUtils/Source/StructUtils/Public/StructView.h b/Engine/Plugins/Experimental/StructUtils/Source/StructUtils/Public/StructView.h
--- a/Engine/Plugins/Experimental/StructUtils/Source/StructUtils/Public/StructView.h
+++ b/Engine/Plugins/Experimental/StructUtils/Source/StructUtils/Public/StructView.h
@@ -112,7 +112,7 @@ public:
 	template <typename OtherType>
 	bool operator==(const OtherType& Other) const
 	{
-		return ((ScriptStruct != Other.GetScriptStruct()) && (StructMemory != Other.GetMemory()));
+		return ((ScriptStruct == Other.GetScriptStruct()) && (StructMemory == Other.GetMemory()));
 	}
 
 	template <typename OtherType>
@@ -236,7 +236,7 @@ public:
 	template <typename OtherType>
 	bool operator==(const OtherType& Other) const
 	{
-		return ((ScriptStruct != Other.GetScriptStruct()) && (StructMemory != Other.GetMemory()));
+		return ((ScriptStruct == Other.GetScriptStruct()) && (StructMemory == Other.GetMemory()));
 	}
 
 	template <typename OtherType>
