Removed unnecessary const.

diff --git a/Engine/Plugins/Experimental/StructUtils/Source/StructUtils/Public/StructView.h b/Engine/Plugins/Experimental/StructUtils/Source/StructUtils/Public/StructView.h
--- a/Engine/Plugins/Experimental/StructUtils/Source/StructUtils/Public/StructView.h
+++ b/Engine/Plugins/Experimental/StructUtils/Source/StructUtils/Public/StructView.h
@@ -179,14 +179,14 @@ public:
 
 	/** Returns const reference to the struct, this getter assumes that all data is valid. */
 	template<typename T>
-	constexpr typename TEnableIf<TIsConst<T>::Value, const T&>::Type Get() const
+	constexpr typename TEnableIf<TIsConst<T>::Value, T&>::Type Get() const
 	{
 		return UE::StructUtils::GetStructRef<T>(ScriptStruct, StructMemory);
 	}
 
 	/** Returns const pointer to the struct, or nullptr if cast is not valid. */
 	template<typename T>
-	constexpr typename TEnableIf<TIsConst<T>::Value, const T*>::Type GetPtr() const
+	constexpr typename TEnableIf<TIsConst<T>::Value, T*>::Type GetPtr() const
 	{
 		return UE::StructUtils::GetStructPtr<T>(ScriptStruct, StructMemory);
 	}
