Refactor FStructArrayView and FConstStructArrayView to be in line with changes to FStructView and exisitng TArrayView.

diff --git a/Engine/Plugins/Experimental/StructUtils/Source/StructUtils/Public/StructArrayView.h b/Engine/Plugins/Experimental/StructUtils/Source/StructUtils/Public/StructArrayView.h
--- a/Engine/Plugins/Experimental/StructUtils/Source/StructUtils/Public/StructArrayView.h
+++ b/Engine/Plugins/Experimental/StructUtils/Source/StructUtils/Public/StructArrayView.h
@@ -3,123 +3,717 @@
 #pragma once
 
 #include "UObject/Class.h"
+#include "StructUtils.h"
+#include "StructView.h"
 
 ///////////////////////////////////////////////////////////////// FStructArrayView /////////////////////////////////////////////////////////////////
 
 /** 
  *	A generic, transient view of a homogeneously-typed array of instances of a specific UScriptStruct
+ *  FStructArrayView supplies mutable access to the elements of the Array.
+ *  Note const FStructArrayView only what the view is pointing to can not be modified the actual struct data is still mutable.
+ *  For a more indepth overview of the constness of views see FConstStructView
  */
 struct FStructArrayView
 {
-	FStructArrayView(FStructArrayView&& Src)
-		: ElementSize(Src.ElementSize)
-		, DataPtr(Src.DataPtr)
-		, NumElements(Src.NumElements)
-		, ElementType(Src.ElementType)
-	{}
+	FStructArrayView() = default;
 
-	FStructArrayView(const FStructArrayView& Src)
-		: ElementSize(Src.ElementSize)
-		, DataPtr(Src.DataPtr)
-		, NumElements(Src.NumElements)
-		, ElementType(Src.ElementType)
-	{}
-	
 	template<typename T>
 	explicit FStructArrayView(TArray<T>& InArray)
-		: ElementSize(sizeof(T))
-		, DataPtr(InArray.GetData())
-		, NumElements(InArray.Num())
-		, ElementType(*StaticStruct<typename TRemoveReference<T>::Type>())
-
+		: DataPtr(InArray.GetData())
+		, ScriptStruct(StaticStruct<typename TRemoveReference<T>::Type>())
+		, ElementSize(sizeof(T))
+		, ArrayNum(InArray.Num())
 	{}
 
 	template<typename T>
 	explicit FStructArrayView(TArrayView<T> InArrayView)
-		: ElementSize(sizeof(T))
-		, DataPtr(InArrayView.GetData())
-		, NumElements(InArrayView.Num())
-		, ElementType(*StaticStruct<typename TRemoveReference<T>::Type>())
+		: DataPtr(InArrayView.GetData())
+		, ScriptStruct(StaticStruct<typename TRemoveReference<T>::Type>())
+		, ElementSize(sizeof(T))
+		, ArrayNum(InArrayView.Num())
 	{}
 
+	FStructArrayView(const UScriptStruct* InScriptStruct, void* InData, const uint32 InElementSize, const int32 InCount)
+		: DataPtr(InData)
+		, ScriptStruct(InScriptStruct)
+		, ElementSize(InElementSize)
+		, ArrayNum(InCount)
+	{
+		check((InData == nullptr) || ((InScriptStruct != nullptr) && (InElementSize > 0) && (InCount >= 0)));
+		check((InData != nullptr) || ((InScriptStruct == nullptr) && (InElementSize == 0) && (InCount == 0)));
+	}
+
 	/**
-	 * Creates a sub view of the given view
+	 * Construct a view with given data.
 	 */
-	FStructArrayView(const FStructArrayView& Src, const int32 OffsetElements, const int32 InNumElements)
-		: ElementSize(Src.ElementSize)
-		, DataPtr((uint8*)Src.DataPtr + Src.ElementSize * OffsetElements)
-		, NumElements(InNumElements)
-		, ElementType(Src.ElementType)
+	FStructArrayView(const UScriptStruct* InScriptStruct, void* InData, const int32 InCount)
+		: DataPtr(InData)
+		, ScriptStruct(InScriptStruct)
+		, ElementSize(InScriptStruct->GetStructureSize())
+		, ArrayNum(InCount)
 	{
-		checkf(InNumElements + OffsetElements <= Src.NumElements, TEXT("Requested range passes over the end of the view, %d + %d> %d")
-			, InNumElements, OffsetElements, Src.NumElements);
+		check(InCount >= 0);
+		check(InData != nullptr || InCount == 0);
 	}
 
-	/** 
-	 * Construct a view with given data.
+	/**
+	 * Checks array invariants: if array size is greater than or equal to zero.
+	 */
+	void CheckInvariants() const
+	{
+		checkSlow(ArrayNum >= 0);
+	}
+
+	/**
+	 * Checks if index is in array range.
+	 *
+	 * @param Index Index to check.
+	 */
+	void RangeCheck(int32 Index) const
+	{
+		checkf((Index >= 0) & (Index < ArrayNum), TEXT("Array index out of bounds: %d from an array of size %d"), Index, ArrayNum);
+	}
+
+	/**
+	 * Checks if a slice range [Index, Index+InNum) is in array range.
+	 * Length is 0 is allowed on empty arrays; Index must be 0 in that case.
+	 *
+	 * @param Index Starting index of the slice.
+	 * @param InNum Length of the slice.
+	 */
+	void SliceRangeCheck(const int32 Index, const int32 InNum) const
+	{
+		checkf(Index >= 0, TEXT("Invalid index (%d)"), Index);
+		checkf(InNum >= 0, TEXT("Invalid count (%d)"), InNum);
+		checkf(Index + InNum <= ArrayNum, TEXT("Range (index: %d, count: %d) lies outside the view of %d elements"), Index, InNum, ArrayNum);
+	}
+
+	/**
+	 * Tests if index is valid, i.e. greater than or equal to zero, and less than the number of elements in the array.
+	 *
+	 * @param Index Index to test.
+	 *
+	 * @returns True if index is valid. False otherwise.
+	 */
+	bool IsValidIndex(const int32 Index) const
+	{
+		return (Index >= 0) && (Index < ArrayNum);
+	}
+
+	/**
+	 * Returns true if the array is empty and contains no elements. 
+	 *
+	 * @returns True if array is empty.
+	 * @see Num
+	 */
+	bool IsEmpty() const
+	{
+		return ArrayNum == 0;
+	}
+
+	/**
+	 * Helper function for returning a pointer to the first array element.
+	 *
+	 * @returns Pointer to first array entry.
+	 */	
+	void* GetData() const
+	{
+		return DataPtr;
+	}
+
+	/**
+	 * Helper function for returning a pointer to an element in the array.
+	 *
+	 * @param Index Index of element, this MUST be a valid element.
+	 * @returns Pointer to array entry at index.
 	 */
-	FStructArrayView(const UScriptStruct& InType, void* InDataPtr, const int32 InNumElements)
-		: ElementSize(InType.GetStructureSize())
-		, DataPtr(InDataPtr)
-		, NumElements(InNumElements)
-		, ElementType(InType)
+	void* GetDataAt(const int32 Index) const 
+	{ 
+		RangeCheck(Index);
+		return (uint8*)DataPtr + ((SIZE_T)Index * ElementSize); 
+	}
+
+	UE_DEPRECATED(5.3, "Use GetDataAt() instead")
+	void* GetMutableDataAt(const int32 Index) const { return GetDataAt(Index); }
+
+	/**
+	 * Helper function returning the size of the inner type.
+	 *
+	 * @returns Size in bytes of array type.
+	 */
+	FORCEINLINE_DEBUGGABLE uint32 GetTypeSize() const
 	{
-		check(InNumElements >= 0);
-		check(InDataPtr != nullptr || InNumElements == 0);
+		return ElementSize;
 	}
 
+	UE_DEPRECATED(5.3, "Removed to bring in to line with ArrayView, use GetTypeSize() instead")
 	SIZE_T GetElementSize() const { return ElementSize; }
-	const void* GetData() const { return DataPtr; }
-	const void* GetDataAt(const int32 Index) const { return (uint8*)DataPtr + Index * ElementSize; }
-	void* GetMutableDataAt(const int32 Index) const { return (uint8*)DataPtr + Index * ElementSize; }
-	int32 Num() const { return NumElements; }
-	const UScriptStruct& GetElementType() const { return ElementType; }
 
-	UE_DEPRECATED(5.2, "This function was renamed to GetElementType. Use that name instead.")
-	const UScriptStruct& GetFragmentType() const { return GetElementType(); }
+	/**
+	 * Returns number of elements in array.
+	 *
+	 * @returns Number of elements in array.
+	 */
+	int32 Num() const { return ArrayNum; }
+
+	UE_DEPRECATED(5.3, "Use GetScriptStruct() instead")
+	const UScriptStruct& GetElementType() const { check(ScriptStruct); return *ScriptStruct; }
+
+	UE_DEPRECATED(5.3, "Use GetScriptStruct() instead")
+	const UScriptStruct& GetFragmentType() const { check(ScriptStruct);  return *ScriptStruct; }
+	
+	/** Returns struct type. */
+	const UScriptStruct* GetScriptStruct() const { return ScriptStruct; }
+
+	/**
+	 * Returns pointer to element at given index. The index MUST be valid. 
+	 * If parametre T is invalid then nullptr will be returned.
+	 *
+	 * @param Index Index of element
+	 * @param T Type of stuct, this must either be the type or a parent type of the struct type in the array.
+	 * @returns Pointer to indexed element.
+	 */
+	template<typename T>
+	T* GetPtrAt(const int32 Index) const
+	{
+		return UE::StructUtils::GetStructPtr<T>(ScriptStruct, GetDataAt(Index)); // GetDataAt() calls RangeCheck().
+	}
+
+	/**
+	 * Returns reference to element at given index. Index and template param T must be must be valid.
+	 *
+	 * @param Index Index of element
+	 * @param T Type of stuct, this must either be the type or a parent type of the struct type in the array.
+	 * @returns Reference to indexed element.
+	 */
+	template<typename T>
+	T& GetAt(const int32 Index) const
+	{
+		return UE::StructUtils::GetStructRef<T>(ScriptStruct, GetDataAt(Index)); // GetDataAt() calls RangeCheck().
+	}
 
 	template<typename T>
+	UE_DEPRECATED(5.3, "Use GetAt() instead!")
 	const T& GetElementAt(const int32 Index) const
 	{
 		return *((T*)GetDataAt(Index));
 	}
 
 	template<typename T>
+	UE_DEPRECATED(5.3, "Use GetAt() instead!")
 	const T& GetElementAtChecked(const int32 Index) const
 	{
-		check(TBaseStructure<T>::Get() == &ElementType);
+		check(TBaseStructure<T>::Get() == &ScriptStruct);
 		return *((T*)GetDataAt(Index));
 	}
 
 	template<typename T>
-	T& GetMutableElementAt(const int32 Index)
+	UE_DEPRECATED(5.3, "Use GetAt() instead!")
+	T& GetMutableElementAt(const int32 Index) const
 	{
 		return *((T*)GetDataAt(Index));
 	}
 
 	template<typename T>
-	T& GetMutableElementAtChecked(const int32 Index)
+	UE_DEPRECATED(5.3, "Use GetAt() instead!")
+	T& GetMutableElementAtChecked(const int32 Index) const
 	{
-		check(TBaseStructure<T>::Get() == &ElementType);
+		check(TBaseStructure<T>::Get() == &ScriptStruct);
 		return *((T*)GetDataAt(Index));
 	}
 
-	FORCEINLINE void Swap(const int32 A, const int32 B)
+	/**
+	 * Array bracket operator. Returns FStructView to element at given index. Index must be valid.
+	 *
+	 * @returns FStructView to indexed element.
+	 */
+	FStructView operator[](const int32 Index) const
 	{
-		checkSlow(A >= 0 && A < NumElements);
-		checkSlow(B >= 0 && B < NumElements);
-		FMemory::Memswap(GetMutableDataAt(A), GetMutableDataAt(B), ElementSize);
+		return FStructView(ScriptStruct, (uint8*)GetDataAt(Index)); // GetDataAt() calls RangeCheck().
+	}
+
+	/**
+	 * Returns n-th last element from the array.
+	 *
+	 * @param IndexFromTheEnd (Optional) Index from the end of array. This MUST index a valid element.
+	 *                        Default is 0.
+	 *
+	 * @returns Reference to n-th last element from the array.
+	 */
+	template<typename T>
+	T& Last(const int32 IndexFromTheEnd = 0) const
+	{
+		return GetAt<T>(ArrayNum - IndexFromTheEnd - 1);
+	}
+
+	/**
+	 * Returns a sliced view
+	 * This is similar to Mid(), but with a narrow contract, i.e. slicing outside of the range of the view is illegal.
+	 *
+	 * @param Index starting index of the new view
+	 * @param InNum number of elements in the new view
+	 * @returns Sliced view
+	 *
+	 * @see Mid
+	 */
+	[[nodiscard]] FStructArrayView Slice(const int32 Index, const int32 InNum) const
+	{
+		SliceRangeCheck(Index, InNum);
+		return FStructArrayView(ScriptStruct, GetDataAt(Index), ElementSize, InNum);
+	}
+
+	/** Returns the left-most part of the view by taking the given number of elements from the left. */
+	[[nodiscard]] FStructArrayView Left(const int32 Count) const
+	{
+		return FStructArrayView(ScriptStruct, DataPtr, ElementSize, FMath::Clamp(Count, 0, ArrayNum));
+	}
+
+	/** Returns the left-most part of the view by chopping the given number of elements from the right. */
+	[[nodiscard]] FStructArrayView LeftChop(const int32 Count) const
+	{
+		return FStructArrayView(ScriptStruct, DataPtr, ElementSize, FMath::Clamp(ArrayNum - Count, 0, ArrayNum));
+	}
+
+	/** Returns the right-most part of the view by taking the given number of elements from the right. */
+	[[nodiscard]] FStructArrayView Right(const int32 Count) const
+	{
+		const int32 OutLen = FMath::Clamp(Count, 0, ArrayNum);
+		return FStructArrayView(ScriptStruct, GetDataAt(ArrayNum - OutLen), ElementSize, OutLen);
+	}
+
+	/** Returns the right-most part of the view by chopping the given number of elements from the left. */
+	[[nodiscard]] FStructArrayView RightChop(const int32 Count) const
+	{
+		const int32 OutLen = FMath::Clamp(ArrayNum - Count, 0, ArrayNum);
+		return FStructArrayView(ScriptStruct, GetDataAt(ArrayNum - OutLen), ElementSize, OutLen);
+	}
+
+	/** Returns the middle part of the view by taking up to the given number of elements from the given position. */
+	[[nodiscard]] FStructArrayView Mid(int32 Index, int32 Count = TNumericLimits<int32>::Max()) const
+	{
+		UE::StructUtils::CalcMidIndexAndCount(ArrayNum, Index, Count);
+		return FStructArrayView(ScriptStruct, GetDataAt(Index), ElementSize, Count);;
+	}
+
+	/** Modifies the view to be the given number of elements from the left. */
+	void LeftInline(const int32 Count)
+	{
+		*this = Left(Count);
+	}
+
+	/** Modifies the view by chopping the given number of elements from the right. */
+	void LeftChopInline(const int32 Count)
+	{
+		*this = LeftChop(Count);
+	}
+
+	/** Modifies the view to be the given number of elements from the right. */
+	void RightInline(const int32 Count)
+	{
+		*this = Right(Count);
+	}
+
+	/** Modifies the view by chopping the given number of elements from the left. */
+	void RightChopInline(const int32 Count)
+	{
+		*this = RightChop(Count);
+	}
+
+	/** Modifies the view to be the middle part by taking up to the given number of elements from the given position. */
+	inline void MidInline(const int32 Position, const int32 Count = TNumericLimits<int32>::Max())
+	{
+		*this = Mid(Position, Count);
+	}
+
+	struct FIterator
+	{
+		FIterator(const FStructArrayView& InOwner, int32 InIndex)
+			: Owner(&InOwner)
+			, Index(InIndex)
+		{}
+
+		FIterator& operator++()
+		{
+			++Index;
+			return *this;
+		}
+
+		FStructView operator*() const
+		{
+			return (*Owner)[Index];
+		}
+
+		FORCEINLINE bool operator == (const FIterator& Other) const
+		{
+			return Owner == Other.Owner
+				&& Index == Other.Index;
+		}
+
+		FORCEINLINE bool operator != (const FIterator& Other) const
+		{
+			return !(*this == Other);
+		}
+
+	private:
+		const FStructArrayView* Owner = nullptr;
+		int32 Index = INDEX_NONE;
+	};
+
+	/** Ranged iteration support. DO NOT USE DIRECTLY. */
+	FORCEINLINE FIterator begin() const { return FIterator(*this, 0); }
+	FORCEINLINE FIterator end() const { return FIterator(*this, Num()); }
+
+	/** Swaps the elements at the specified Indicies. The indicies must be valid! */
+	FORCEINLINE void Swap(const int32 Index1, const int32 Index2)
+	{
+		FMemory::Memswap(GetDataAt(Index1), GetDataAt(Index2), ElementSize); // GetDataAt() calls RangeCheck().
 	}
 
 private:
-	SIZE_T ElementSize = 0;
 	void* DataPtr = nullptr;
-	int32 NumElements = 0;
-	const UScriptStruct& ElementType;
+	const UScriptStruct* ScriptStruct = nullptr;
+	uint32 ElementSize = 0;
+	int32 ArrayNum = 0;
+};
+
+///////////////////////////////////////////////////////////////// FConstStructArrayView /////////////////////////////////////////////////////////////////
+
+/**
+ *	A generic, transient view of a homogeneously-typed array of instances of a specific UScriptStruct
+ *  FConstStructArrayView supplies immutable access to the elements of the Array.
+ *  For a more indepth overview of the constness of views see FConstStructView 
+ */
+struct FConstStructArrayView
+{
+	FConstStructArrayView() = default;
+
+	template<typename T>
+	explicit FConstStructArrayView(TArray<T>& InArray)
+		: DataPtr(InArray.GetData())
+		, ScriptStruct(StaticStruct<typename TRemoveReference<T>::Type>())
+		, ElementSize(sizeof(T))
+		, ArrayNum(InArray.Num())
+	{}
+
+	template<typename T>
+	explicit FConstStructArrayView(TArrayView<T> InArrayView)
+		: DataPtr(InArrayView.GetData())
+		, ScriptStruct(StaticStruct<typename TRemoveReference<T>::Type>())
+		, ElementSize(sizeof(T))
+		, ArrayNum(InArrayView.Num())
+	{}
+
+	FConstStructArrayView(const FStructArrayView Src)
+		: DataPtr(Src.GetData())
+		, ScriptStruct(Src.GetScriptStruct())
+		, ElementSize(Src.GetTypeSize())
+		, ArrayNum(Src.Num())
+	{}
+
+	FConstStructArrayView(const UScriptStruct* InScriptStruct, const void* InData, const uint32 InElementSize, const int32 InCount)
+		: DataPtr(InData)
+		, ScriptStruct(InScriptStruct)
+		, ElementSize(InElementSize)
+		, ArrayNum(InCount)
+	{
+		check((InData == nullptr) || ((InScriptStruct != nullptr) && (InElementSize > 0) && (InCount >= 0)));
+		check((InData != nullptr) || ((InScriptStruct == nullptr) && (InElementSize == 0) && (InCount == 0)));
+	}
+
+	/**
+	 * Construct a view with given data.
+	 */
+	FConstStructArrayView(const UScriptStruct* InScriptStruct, const void* InData, const int32 InCount)
+		: DataPtr(InData)
+		, ScriptStruct(InScriptStruct)
+		, ElementSize(InScriptStruct->GetStructureSize())
+		, ArrayNum(InCount)
+	{
+		check(InCount >= 0);
+		check(InData != nullptr || InCount == 0);
+	}
+
+	FConstStructArrayView& operator=(const FStructArrayView StructArrayView)
+	{
+		*this = FConstStructArrayView(StructArrayView);
+		return *this;
+	}
+
+	/**
+	 * Checks array invariants: if array size is greater than or equal to zero.
+	 */
+	void CheckInvariants() const
+	{
+		checkSlow(ArrayNum >= 0);
+	}
+
+	/**
+	 * Checks if index is in array range.
+	 *
+	 * @param Index Index to check.
+	 */
+	void RangeCheck(int32 Index) const
+	{
+		checkf((Index >= 0) & (Index < ArrayNum), TEXT("Array index out of bounds: %d from an array of size %d"), Index, ArrayNum);
+	}
+
+	/**
+	 * Checks if a slice range [Index, Index+InNum) is in array range.
+	 * Length is 0 is allowed on empty arrays; Index must be 0 in that case.
+	 *
+	 * @param Index Starting index of the slice.
+	 * @param InNum Length of the slice.
+	 */
+	FORCEINLINE void SliceRangeCheck(int32 Index, int32 InNum) const
+	{
+		checkf(Index >= 0, TEXT("Invalid index (%d)"), Index);
+		checkf(InNum >= 0, TEXT("Invalid count (%d)"), InNum);
+		checkf(Index + InNum <= ArrayNum, TEXT("Range (index: %d, count: %d) lies outside the view of %d elements"), Index, InNum, ArrayNum);
+	}
+
+	/**
+	 * Tests if index is valid, i.e. greater than or equal to zero, and less than the number of elements in the array.
+	 *
+	 * @param Index Index to test.
+	 *
+	 * @returns True if index is valid. False otherwise.
+	 */
+	bool IsValidIndex(int32 Index) const
+	{
+		return (Index >= 0) && (Index < ArrayNum);
+	}
+
+	/**
+	 * Returns true if the array is empty and contains no elements. 
+	 *
+	 * @returns True if array is empty.
+	 * @see Num
+	 */
+	bool IsEmpty() const
+	{
+		return ArrayNum == 0;
+	}
+
+	/**
+	 * Helper function for returning a pointer to the first array entry.
+	 *
+	 * @returns const Pointer to first array entry.
+	 */
+	const void* GetData() const
+	{
+		return DataPtr;
+	}
+
+	/**
+	 * Helper function for returning a pointer to an element in the array.
+	 *
+	 * @param Index Index of element, this MUST be a valid element.
+	 * @returns const Pointer to entry.
+	 */
+	const void* GetDataAt(const int32 Index) const
+	{
+		RangeCheck(Index);
+		return (uint8*)DataPtr + (SIZE_T)Index * ElementSize;
+	}
+
+	/**
+	 * Helper function returning the size of the inner type.
+	 *
+	 * @returns Size in bytes of array type.
+	 */
+	FORCEINLINE_DEBUGGABLE uint32 GetTypeSize()
+	{
+		return ElementSize;
+	}
+
+	/**
+	 * Returns number of elements in array.
+	 *
+	 * @returns Number of elements in array.
+	 */
+	int32 Num() const { return ArrayNum; }
+
+	/** Returns struct type. */
+	const UScriptStruct* GetScriptStruct() const { return ScriptStruct; }
+
+	/**
+	 * Returns pointer to element at given index. The index MUST be valid. 
+	 * If parametre T is invalid then nullptr will be returned.
+	 *
+	 * @param Index Index of element
+	 * @param T Type of stuct, this must either be the type or a parent type of the struct type in the array.
+	 * @returns const pointer to indexed element.
+	 */
+	template<typename T>
+	constexpr typename TEnableIf<TIsConst<T>::Value, T*>::Type GetPtrAt(const int32 Index) const
+	{
+		return UE::StructUtils::GetStructPtr<T>(ScriptStruct, GetDataAt(Index)); // GetDataAt() calls RangeCheck().
+	}
+
+	/**
+	 * Returns reference to element at given index. Index and template param T MUST be must be valid.
+	 *
+	 * @param Index Index of element
+	 * @param T Type of stuct, this must either be the type or a parent type of the struct type in the array.
+	 * @returns const reference to indexed element.
+	 */
+	template<typename T>
+	constexpr typename TEnableIf<TIsConst<T>::Value, T&>::Type GetAt(const int32 Index) const
+	{
+		return UE::StructUtils::GetStructRef<T>(ScriptStruct, GetDataAt(Index)); // GetDataAt() calls RangeCheck().
+	}
+
+	/**
+	 * Array bracket operator. Returns FConstStructView to element at given index. Index must be valid.
+	 *
+	 * @returns FConstStructView to indexed element.
+	 */
+	FConstStructView operator[](int32 Index) const
+	{
+		return FConstStructView(ScriptStruct, (const uint8*)GetDataAt(Index)); // GetDataAt() calls RangeCheck().
+	}
+
+	/**
+	 * Returns n-th last element from the array.
+	 *
+	 * @param IndexFromTheEnd (Optional) Index from the end of array. This MUST index a valid element.
+	 *                        Default is 0.
+	 *
+	 * @returns const reference to n-th last element from the array.
+	 */
+	template<typename T>
+	constexpr typename TEnableIf<TIsConst<T>::Value, T&>::Type Last(int32 IndexFromTheEnd = 0) const
+	{
+		return GetAt<T>(ArrayNum - IndexFromTheEnd - 1);
+	}
+
+	/**
+	 * Returns a sliced view
+	 * This is similar to Mid(), but with a narrow contract, i.e. slicing outside of the range of the view is illegal.
+	 *
+	 * @param Index starting index of the new view
+	 * @param InNum number of elements in the new view
+	 * @returns Sliced view
+	 *
+	 * @see Mid
+	 */
+	[[nodiscard]] FConstStructArrayView Slice(const int32 Index, const int32 InNum) const
+	{
+		SliceRangeCheck(Index, InNum);
+		return FConstStructArrayView(ScriptStruct, GetDataAt(Index), ElementSize, InNum);
+	}
+
+	/** Returns the left-most part of the view by taking the given number of elements from the left. */
+	[[nodiscard]] FConstStructArrayView Left(const int32 Count) const
+	{
+		return FConstStructArrayView(ScriptStruct, DataPtr, ElementSize, FMath::Clamp(Count, 0, ArrayNum));
+	}
+
+	/** Returns the left-most part of the view by chopping the given number of elements from the right. */
+	[[nodiscard]] FConstStructArrayView LeftChop(const int32 Count) const
+	{
+		return FConstStructArrayView(ScriptStruct, DataPtr, ElementSize, FMath::Clamp(ArrayNum - Count, 0, ArrayNum));
+	}
+
+	/** Returns the right-most part of the view by taking the given number of elements from the right. */
+	[[nodiscard]] FConstStructArrayView Right(const int32 Count) const
+	{
+		const int32 OutLen = FMath::Clamp(Count, 0, ArrayNum);
+		return FConstStructArrayView(ScriptStruct, GetDataAt(ArrayNum - OutLen), ElementSize, OutLen);
+	}
+
+	/** Returns the right-most part of the view by chopping the given number of elements from the left. */
+	[[nodiscard]] FConstStructArrayView RightChop(const int32 Count) const
+	{
+		const int32 OutLen = FMath::Clamp(ArrayNum - Count, 0, ArrayNum);
+		return FConstStructArrayView(ScriptStruct, GetDataAt(ArrayNum - OutLen), ElementSize, OutLen);
+	}
+
+	/** Returns the middle part of the view by taking up to the given number of elements from the given position. */
+	[[nodiscard]] FConstStructArrayView Mid(int32 Index, int32 Count = TNumericLimits<int32>::Max()) const
+	{
+		UE::StructUtils::CalcMidIndexAndCount(ArrayNum, Index, Count);
+		return FConstStructArrayView(ScriptStruct, GetDataAt(Index), ElementSize, Count);;
+	}
+
+	/** Modifies the view to be the given number of elements from the left. */
+	void LeftInline(const int32 Count)
+	{
+		*this = Left(Count);
+	}
+
+	/** Modifies the view by chopping the given number of elements from the right. */
+	void LeftChopInline(const int32 Count)
+	{
+		*this = LeftChop(Count);
+	}
+
+	/** Modifies the view to be the given number of elements from the right. */
+	void RightInline(const int32 Count)
+	{
+		*this = Right(Count);
+	}
+
+	/** Modifies the view by chopping the given number of elements from the left. */
+	void RightChopInline(const int32 Count)
+	{
+		*this = RightChop(Count);
+	}
+
+	/** Modifies the view to be the middle part by taking up to the given number of elements from the given position. */
+	inline void MidInline(const int32 Position, const int32 Count = TNumericLimits<int32>::Max())
+	{
+		*this = Mid(Position, Count);
+	}
+
+	struct FIterator
+	{
+		FIterator(const FConstStructArrayView& InOwner, int32 InIndex)
+			: Owner(&InOwner)
+			, Index(InIndex)
+		{}
+
+		FIterator& operator++()
+		{
+			++Index;
+			return *this;
+		}
+
+		FConstStructView operator*() const
+		{
+			return (*Owner)[Index];
+		}
+
+		FORCEINLINE bool operator == (const FIterator& Other) const
+		{
+			return Owner == Other.Owner
+				&& Index == Other.Index;
+		}
+
+		FORCEINLINE bool operator != (const FIterator& Other) const
+		{
+			return !(*this == Other);
+		}
+
+	private:
+		const FConstStructArrayView* Owner = nullptr;
+		int32 Index = INDEX_NONE;
+	};
+
+	/** Ranged iteration support. DO NOT USE DIRECTLY. */
+	FORCEINLINE FIterator begin() const { return FIterator(*this, 0); }
+	FORCEINLINE FIterator end() const { return FIterator(*this, Num()); }
+
+private:
+	const void* DataPtr = nullptr;
+	const UScriptStruct* ScriptStruct = nullptr;
+	uint32 ElementSize = 0;
+	int32 ArrayNum = 0;
 };
 
 #if UE_ENABLE_INCLUDE_ORDER_DEPRECATED_IN_5_2
 #include "InstancedStruct.h"
 #include "SharedStruct.h"
 #include "StructUtils.h"
-#endif
+#endif
\ No newline at end of file
diff --git a/Engine/Plugins/Experimental/StructUtils/Source/StructUtils/Public/StructUtils.h b/Engine/Plugins/Experimental/StructUtils/Source/StructUtils/Public/StructUtils.h
--- a/Engine/Plugins/Experimental/StructUtils/Source/StructUtils/Public/StructUtils.h
+++ b/Engine/Plugins/Experimental/StructUtils/Source/StructUtils/Public/StructUtils.h
@@ -18,7 +18,7 @@ namespace UE::StructUtils
 
 	/** Returns reference to the struct, this assumes that all data is valid. */
 	template<typename T>
-	T& GetStructRef(const UScriptStruct* ScriptStruct, uint8* StructMemory)
+	T& GetStructRef(const UScriptStruct* ScriptStruct, void* StructMemory)
 	{
 		check(StructMemory != nullptr);
 		check(ScriptStruct != nullptr);
@@ -28,7 +28,7 @@ namespace UE::StructUtils
 
 	/** Returns pointer to the struct, or nullptr if cast is not valid. */
 	template<typename T>
-	T* GetStructPtr(const UScriptStruct* ScriptStruct, uint8* StructMemory)
+	T* GetStructPtr(const UScriptStruct* ScriptStruct, void* StructMemory)
 	{
 		if (StructMemory != nullptr 
 			&& ScriptStruct 
@@ -42,15 +42,33 @@ namespace UE::StructUtils
 
 	/** Returns const reference to the struct, this assumes that all data is valid. */
 	template<typename T>
-	const T& GetStructRef(const UScriptStruct* ScriptStruct, const uint8* StructMemory)
+	const T& GetStructRef(const UScriptStruct* ScriptStruct, const void* StructMemory)
 	{
-		return GetStructRef<T>(ScriptStruct, const_cast<uint8*>(StructMemory));
+		return GetStructRef<T>(ScriptStruct, const_cast<void*>(StructMemory));
 	}
 
 	/** Returns const pointer to the struct, or nullptr if cast is not valid. */
 	template<typename T>
-	const T* GetStructPtr(const UScriptStruct* ScriptStruct, const uint8* StructMemory)
+	const T* GetStructPtr(const UScriptStruct* ScriptStruct, const void* StructMemory)
 	{
-		return GetStructPtr<T>(ScriptStruct, const_cast<uint8*>(StructMemory));
+		return GetStructPtr<T>(ScriptStruct, const_cast<void*>(StructMemory));
+	}
+
+	/** 
+	 *  Returns the middle part of an array or view by taking up to the given number of elements from the given position.
+	 *  Based on TArrayView::Mid().
+	 */
+	inline void CalcMidIndexAndCount(int32 ArrayNum, int32& InOutIndex, int32& InOutCount)
+	{
+		// Clamp minimum index at the start of the range, adjusting the length down if necessary
+		const int32 NegativeIndexOffset = (InOutIndex < 0) ? InOutIndex : 0;
+		InOutCount += NegativeIndexOffset;
+		InOutIndex -= NegativeIndexOffset;
+
+		// Clamp maximum index at the end of the range
+		InOutIndex = (InOutIndex > ArrayNum) ? ArrayNum : InOutIndex;
+
+		// Clamp count between 0 and the distance to the end of the range
+		InOutCount = FMath::Clamp(InOutCount, 0, (ArrayNum - InOutIndex));
 	}
 }
diff --git a/Engine/Plugins/Experimental/StructUtils/Source/StructUtilsTestSuite/Private/StructArrayViewTest.cpp b/Engine/Plugins/Experimental/StructUtils/Source/StructUtilsTestSuite/Private/StructArrayViewTest.cpp
--- a/Engine/Plugins/Experimental/StructUtils/Source/StructUtilsTestSuite/Private/StructArrayViewTest.cpp
+++ b/Engine/Plugins/Experimental/StructUtils/Source/StructUtilsTestSuite/Private/StructArrayViewTest.cpp
@@ -9,7 +9,7 @@
 namespace FStructArrayViewTest
 {
 
-struct FTest_ChangeElements : FAITestBase
+struct FTest_ArrayViewChangeElements : FAITestBase
 {
 	virtual bool InstantTest() override
 	{
@@ -20,25 +20,261 @@ struct FTest_ChangeElements : FAITestBase
 		{
 			ValuesF.Add(Value);
 		}
-		FStructArrayView ViewF = FStructArrayView(ValuesF);
+		FStructArrayView ViewF(ValuesF);
 
-		for (int i = 0; i < ValuesF.Num(); ++i)
+		for (int32 i = 0; i < ValuesF.Num(); ++i)
 		{
-			ViewF.GetMutableElementAt<FTestStructSimple>(i).Float += 10.f;
+			ViewF.GetAt<FTestStructSimple>(i).Float += 10.f;
 		}
 
-		for (int i = 0; i < ValuesF.Num(); ++i)
 		{
-			AITEST_EQUAL(TEXT("Element should have the same value independently of access method"), ValuesF[i].Float, ViewF.GetElementAt<FTestStructSimple>(i).Float);
-			AITEST_EQUAL(TEXT("Element should reflect the change done to it"), ViewF.GetElementAt<FTestStructSimple>(i).Float, OriginalValues[i] + Increment);
-			AITEST_EQUAL(TEXT("Element should point at the same adress in memory, independently of access method"), (void*)&ValuesF[i], ViewF.GetDataAt(i));
+			int32 i = 0;
+			for (FStructView StructViewTest : ViewF)
+			{
+				AITEST_EQUAL(TEXT("FStructView should have the same value as TArray"), ValuesF[i].Float, StructViewTest.Get<FTestStructSimple>().Float);
+				AITEST_EQUAL(TEXT("FStructArrayView GetAt() should have the same value as TArray"), ValuesF[i].Float, ViewF.GetAt< FTestStructSimple>(i).Float);
+				AITEST_TRUE(TEXT("FStructArrayView GetPtrAt() should be non Null"), ViewF.GetPtrAt< FTestStructSimple>(i) != nullptr);
+				AITEST_EQUAL(TEXT("FStructArrayView GetPtrAt() should have the same value as TArray"), ValuesF[i].Float, ViewF.GetPtrAt< FTestStructSimple>(i)->Float);
+				AITEST_EQUAL(TEXT("FStructArrayView operator[] should have the same value as TArray"), ValuesF[i].Float, ViewF[i].Get<FTestStructSimple>().Float);
+				AITEST_EQUAL(TEXT("FStructArrayView Element should reflect the change done to it"), ViewF.GetAt<FTestStructSimple>(i).Float, OriginalValues[i] + Increment);
+				AITEST_EQUAL(TEXT("FStructArrayView Element should point at the same adress in memory as TArray"), (void*)&ValuesF[i], ViewF.GetDataAt(i));
+
+				++i;
+			}
+			AITEST_EQUAL(TEXT("FStructArrayView should have the same number of elements value as TArray"), i, ValuesF.Num());
+		}
+		
+		{
+			FStructArrayView ViewLeft(ViewF);
+			const int32 LeftNum = 3;
+			ViewLeft.LeftInline(LeftNum);
+			AITEST_EQUAL(TEXT("Incorrect LeftInline() modified FStructArrayView Num"), LeftNum, ViewLeft.Num());
+
+			int32 i = 0;
+			for (FStructView StructViewTest : ViewLeft)
+			{
+				AITEST_EQUAL(TEXT("LeftInline() modified FStructArrayView element should be the same as correpsonding element in TArray"), StructViewTest.Get<FTestStructSimple>().Float, ValuesF[i].Float);
+
+				++i;
+			}
+		}
+
+		{
+			FStructArrayView ViewLeftChop(ViewF);
+			const int32 LeftChopNum = 4;
+			ViewLeftChop.LeftChopInline(LeftChopNum);
+			AITEST_EQUAL(TEXT("Incorrect LeftChopInline() modified FStructArrayView Num"), ViewF.Num() - LeftChopNum, ViewLeftChop.Num());
+
+			int32 i = 0;
+			for (FStructView StructViewTest : ViewLeftChop)
+			{
+				AITEST_EQUAL(TEXT("LeftChopInline() modified FStructArrayView element should be the same as correpsonding element in TArray"), StructViewTest.Get<FTestStructSimple>().Float, ValuesF[i].Float);
+
+				++i;
+			}
+		}
+
+		{
+			FStructArrayView ViewRight(ViewF);
+			const int32 RightNum = 2;
+			ViewRight.RightInline(RightNum);
+			AITEST_EQUAL(TEXT("Incorrect RightInline() modified FStructArrayView Num"), RightNum, ViewRight.Num());
+
+			int32 i = ViewF.Num() - RightNum;
+			for (FStructView StructViewTest : ViewRight)
+			{
+				AITEST_EQUAL(TEXT("RightInline() modified FStructArrayView element should be the same as correpsonding element in TArray"), StructViewTest.Get<FTestStructSimple>().Float, ValuesF[i].Float);
+
+				++i;
+			}
+		}
+
+		{
+			FStructArrayView ViewRightChop(ViewF);
+			const int32 RightChopNum = 6;
+			ViewRightChop.RightChopInline(RightChopNum);
+			AITEST_EQUAL(TEXT("Incorrect RightChopInline() modified FStructArrayView Num"), ViewF.Num() - RightChopNum, ViewRightChop.Num());
+
+			int32 i = RightChopNum;
+			for (FStructView StructViewTest : ViewRightChop)
+			{
+				AITEST_EQUAL(TEXT("RightChopInline() modified FStructArrayView element should be the same as correpsonding element in TArray"), StructViewTest.Get<FTestStructSimple>().Float, ValuesF[i].Float);
+
+				++i;
+			}
+		}
+
+		{
+			FStructArrayView ViewMid(ViewF);
+			const int32 MidPositon = 3;
+			const int32 MidNum = 4;
+			ViewMid.MidInline(MidPositon, MidNum);
+			AITEST_EQUAL(TEXT("Incorrect MidInline() modified FStructArrayView Num"), MidNum, ViewMid.Num());
+
+			int32 i = MidPositon;
+			for (FStructView StructViewTest : ViewMid)
+			{
+				AITEST_EQUAL(TEXT("MidInline() modified FStructArrayView element should be the same as correpsonding element in TArray"), StructViewTest.Get<FTestStructSimple>().Float, ValuesF[i].Float);
+
+				++i;
+			}
+		}
+
+		{
+			FStructArrayView ViewSlice(ViewF);
+			const int32 SlicePositon = 2;
+			const int32 SliceNum = 5;
+			ViewSlice = ViewSlice.Slice(SlicePositon, SliceNum);
+			AITEST_EQUAL(TEXT("Incorrect Slice() modified FStructArrayView Num"), SliceNum, ViewSlice.Num());
+
+			int32 i = SlicePositon;
+			for (FStructView StructViewTest : ViewSlice)
+			{
+				AITEST_EQUAL(TEXT("Slice() modified FStructArrayView element should be the same as correpsonding element in TArray"), StructViewTest.Get<FTestStructSimple>().Float, ValuesF[i].Float);
+
+				++i;
+			}
+		}
+
+		return true;
+	}
+};
+
+IMPLEMENT_AI_INSTANT_TEST(FTest_ArrayViewChangeElements, "System.StructUtils.StructArrayView.ChangeElements");
+
+struct FTest_ConstArrayViewChangeElements : FAITestBase
+{
+	virtual bool InstantTest() override
+	{
+		constexpr float Increment = 10.f;
+		TArray<float> OriginalValues = { 1,2,3,4,5,6,7,8,9,10 };
+		TArray<FTestStructSimple> ValuesF;
+		for (float Value : OriginalValues)
+		{
+			ValuesF.Add(Value);
+		}
+		FStructArrayView ViewF(ValuesF);
+		FConstStructArrayView ViewConstF(ViewF);
+
+		for (int32 i = 0; i < ValuesF.Num(); ++i)
+		{
+			ViewF.GetAt<FTestStructSimple>(i).Float += 10.f;
+		}
+
+		{
+			int32 i = 0;
+			for (FConstStructView StructViewTest : ViewConstF)
+			{
+				AITEST_EQUAL(TEXT("FConstStructView should have the same value as TArray"), ValuesF[i].Float, StructViewTest.Get<const FTestStructSimple>().Float);
+				AITEST_EQUAL(TEXT("FConstStructArrayView GetAt() should have the same value as TArray"), ValuesF[i].Float, ViewConstF.GetAt<const FTestStructSimple>(i).Float);
+				AITEST_TRUE(TEXT("FConstStructArrayView GetPtrAt() should be non Null"), ViewConstF.GetPtrAt<const FTestStructSimple>(i) != nullptr);
+				AITEST_EQUAL(TEXT("FConstStructArrayView GetPtrAt() should have the same value as TArray"), ValuesF[i].Float, ViewConstF.GetPtrAt<const FTestStructSimple>(i)->Float);
+				AITEST_EQUAL(TEXT("FConstStructArrayView operator[] should have the same value as TArray"), ValuesF[i].Float, ViewConstF[i].Get<const FTestStructSimple>().Float);
+				AITEST_EQUAL(TEXT("FConstStructArrayView Element should reflect the change done to it"), ViewConstF.GetAt<const FTestStructSimple>(i).Float, OriginalValues[i] + Increment);
+				AITEST_EQUAL(TEXT("FConstStructArrayView Element should point at the same adress in memory as TArray"), (void*)&ValuesF[i], ViewConstF.GetDataAt(i));
+
+				++i;
+			}
+			AITEST_EQUAL(TEXT("FConstStructArrayView should have the same number of elements value as TArray"), i, ViewConstF.Num());
+		}
+
+		{
+			FConstStructArrayView ViewLeft(ViewF);
+			const int32 LeftNum = 3;
+			ViewLeft.LeftInline(LeftNum);
+			AITEST_EQUAL(TEXT("Incorrect LeftInline() modified FConstStructArrayView Num"), LeftNum, ViewLeft.Num());
+
+			int32 i = 0;
+			for (FConstStructView StructViewTest : ViewLeft)
+			{
+				AITEST_EQUAL(TEXT("LeftInline() modified FConstStructArrayView element should be the same as correpsonding element in TArray"), StructViewTest.Get<const FTestStructSimple>().Float, ValuesF[i].Float);
+
+				++i;
+			}
+		}
+
+		{
+			FConstStructArrayView ViewLeftChop(ViewF);
+			const int32 LeftChopNum = 4;
+			ViewLeftChop.LeftChopInline(LeftChopNum);
+			AITEST_EQUAL(TEXT("Incorrect LeftChopInline() modified FConstStructArrayView Num"), ViewF.Num() - LeftChopNum, ViewLeftChop.Num());
+
+			int32 i = 0;
+			for (FConstStructView StructViewTest : ViewLeftChop)
+			{
+				AITEST_EQUAL(TEXT("LeftChopInline() modified FConstStructArrayView element should be the same as correpsonding element in TArray"), StructViewTest.Get<const FTestStructSimple>().Float, ValuesF[i].Float);
+
+				++i;
+			}
+		}
+
+		{
+			FConstStructArrayView ViewRight(ViewF);
+			const int32 RightNum = 2;
+			ViewRight.RightInline(RightNum);
+			AITEST_EQUAL(TEXT("Incorrect RightInline() modified FConstStructArrayView Num"), RightNum, ViewRight.Num());
+
+			int32 i = ViewF.Num() - RightNum;
+			for (FConstStructView StructViewTest : ViewRight)
+			{
+				AITEST_EQUAL(TEXT("RightInline() modified FConstStructArrayView element should be the same as correpsonding element in TArray"), StructViewTest.Get<const FTestStructSimple>().Float, ValuesF[i].Float);
+
+				++i;
+			}
+		}
+
+		{
+			FConstStructArrayView ViewRightChop(ViewF);
+			const int32 RightChopNum = 6;
+			ViewRightChop.RightChopInline(RightChopNum);
+			AITEST_EQUAL(TEXT("Incorrect RightChopInline() modified FConstStructArrayView Num"), ViewF.Num() - RightChopNum, ViewRightChop.Num());
+
+			int32 i = RightChopNum;
+			for (FConstStructView StructViewTest : ViewRightChop)
+			{
+				AITEST_EQUAL(TEXT("RightChopInline() modified FConstStructArrayView element should be the same as correpsonding element in TArray"), StructViewTest.Get<const FTestStructSimple>().Float, ValuesF[i].Float);
+
+				++i;
+			}
+		}
+
+		{
+			FConstStructArrayView ViewMid(ViewF);
+			const int32 MidPositon = 3;
+			const int32 MidNum = 4;
+			ViewMid.MidInline(MidPositon, MidNum);
+			AITEST_EQUAL(TEXT("Incorrect MidInline() modified FConstStructArrayView Num"), MidNum, ViewMid.Num());
+
+			int32 i = MidPositon;
+			for (FConstStructView StructViewTest : ViewMid)
+			{
+				AITEST_EQUAL(TEXT("MidInline() modified FConstStructArrayView element should be the same as correpsonding element in TArray"), StructViewTest.Get<const FTestStructSimple>().Float, ValuesF[i].Float);
+
+				++i;
+			}
+		}
+
+		{
+			FConstStructArrayView ViewSlice(ViewF);
+			const int32 SlicePositon = 2;
+			const int32 SliceNum = 5;
+			ViewSlice = ViewSlice.Slice(SlicePositon, SliceNum);
+			AITEST_EQUAL(TEXT("Incorrect Slice() modified FConstStructArrayView Num"), SliceNum, ViewSlice.Num());
+
+			int32 i = SlicePositon;
+			for (FConstStructView StructViewTest : ViewSlice)
+			{
+				AITEST_EQUAL(TEXT("Slice() modified FConstStructArrayView element should be the same as correpsonding element in TArray"), StructViewTest.Get<const FTestStructSimple>().Float, ValuesF[i].Float);
+
+				++i;
+			}
 		}
 
 		return true;
 	}
 };
-IMPLEMENT_AI_INSTANT_TEST(FTest_ChangeElements, "System.StructUtils.StructArrayView.ChangeElements");
 
+IMPLEMENT_AI_INSTANT_TEST(FTest_ConstArrayViewChangeElements, "System.StructUtils.StructArrayView.ChangeElementsConst");
 } // FStructArrayViewTest
 
 #undef LOCTEXT_NAMESPACE
diff --git a/Engine/Plugins/Runtime/MassEntity/Source/MassEntity/Private/MassArchetypeData.cpp b/Engine/Plugins/Runtime/MassEntity/Source/MassEntity/Private/MassArchetypeData.cpp
--- a/Engine/Plugins/Runtime/MassEntity/Source/MassEntity/Private/MassArchetypeData.cpp
+++ b/Engine/Plugins/Runtime/MassEntity/Source/MassEntity/Private/MassArchetypeData.cpp
@@ -1178,13 +1178,14 @@ void FMassArchetypeData::BatchSetFragmentValues(TConstArrayView<FMassArchetypeEn
 			FStructArrayView FragmentPayload = Payload[i];
 			check(FragmentPayload.Num() - EntitiesHandled >= EntityRange.Length);
 
-			const UScriptStruct& FragmentType = FragmentPayload.GetElementType();
+			const UScriptStruct* FragmentType = FragmentPayload.GetScriptStruct();
+			check(FragmentType);
 
-			const int32 FragmentIndex = FragmentIndexMap.FindChecked(&FragmentType);
+			const int32 FragmentIndex = FragmentIndexMap.FindChecked(FragmentType);
 			void* Dst = FragmentConfigs[FragmentIndex].GetFragmentData(Chunk.GetRawMemory(), EntityRange.SubchunkStart);
 			const void* Src = FragmentPayload.GetDataAt(EntitiesHandled);
 
-			FragmentType.CopyScriptStruct(Dst, Src, EntityRange.Length);
+			FragmentType->CopyScriptStruct(Dst, Src, EntityRange.Length);
 		}
 
 		EntitiesHandled += EntityRange.Length;
diff --git a/Engine/Plugins/Runtime/MassEntity/Source/MassEntity/Public/MassEntityTypes.h b/Engine/Plugins/Runtime/MassEntity/Source/MassEntity/Public/MassEntityTypes.h
--- a/Engine/Plugins/Runtime/MassEntity/Source/MassEntity/Public/MassEntityTypes.h
+++ b/Engine/Plugins/Runtime/MassEntity/Source/MassEntity/Public/MassEntityTypes.h
@@ -400,7 +400,7 @@ struct FMassGenericPayloadViewSlice
 
 	FStructArrayView operator[](const int32 Index) const
 	{
-		return FStructArrayView(Source.Content[Index], StartIndex, Count);
+		return Source.Content[Index].Slice(StartIndex, Count);
 	}
 
 	/** @return the number of "layers" (i.e. number of original arrays) this payload has been built from */
diff --git a/Engine/Plugins/Runtime/MassEntity/Source/MassEntityTestSuite/Private/MassUtilsTest.cpp b/Engine/Plugins/Runtime/MassEntity/Source/MassEntityTestSuite/Private/MassUtilsTest.cpp
--- a/Engine/Plugins/Runtime/MassEntity/Source/MassEntityTestSuite/Private/MassUtilsTest.cpp
+++ b/Engine/Plugins/Runtime/MassEntity/Source/MassEntityTestSuite/Private/MassUtilsTest.cpp
@@ -72,8 +72,8 @@ struct FUtilsTest_MultiSort_GenericArray : FExecutionTestBase
 
 		for (int i = 0; i < ValuesF.Num(); ++i)
 		{
-			AITEST_EQUAL(TEXT("Element in ValuesU should represent the absolute value of the element in ViewF"), float(FMath::Abs(ValuesU[i])), ViewF.GetElementAt<FTestFragment_Float>(i).Value);
-			AITEST_EQUAL(TEXT("Element in ViewB should represent the sign of the element in ValuesU"), ValuesU[i] < 0, ViewB.GetElementAt<FTestFragment_Bool>(i).bValue);
+			AITEST_EQUAL(TEXT("Element in ValuesU should represent the absolute value of the element in ViewF"), float(FMath::Abs(ValuesU[i])), ViewF.GetAt<const FTestFragment_Float>(i).Value);
+			AITEST_EQUAL(TEXT("Element in ViewB should represent the sign of the element in ValuesU"), ValuesU[i] < 0, ViewB.GetAt<const FTestFragment_Bool>(i).bValue);
 		}
 
 		return true;
@@ -105,8 +105,8 @@ struct FUtilsTest_MultiSort_Payload : FExecutionTestBase
 
 		for (int i = 0; i < ValuesF.Num(); ++i)
 		{
-			AITEST_EQUAL(TEXT("Element in ValuesU should represent the absolute value of the element in ViewF"), float(FMath::Abs(ValuesU[i])), ViewF.GetElementAt<FTestFragment_Float>(i).Value);
-			AITEST_EQUAL(TEXT("Element in ViewB should represent the sign of the element in ValuesU"), ValuesU[i] < 0, ViewB.GetElementAt<FTestFragment_Bool>(i).bValue);
+			AITEST_EQUAL(TEXT("Element in ValuesU should represent the absolute value of the element in ViewF"), float(FMath::Abs(ValuesU[i])), ViewF.GetAt<const FTestFragment_Float>(i).Value);
+			AITEST_EQUAL(TEXT("Element in ViewB should represent the sign of the element in ValuesU"), ValuesU[i] < 0, ViewB.GetAt<const FTestFragment_Bool>(i).bValue);
 		}
 
 		return true;
