Removed pointer from a constructor and replaced with a reference.

diff --git a/Engine/Plugins/Experimental/StructUtils/Source/StructUtils/Public/StructArrayView.h b/Engine/Plugins/Experimental/StructUtils/Source/StructUtils/Public/StructArrayView.h
--- a/Engine/Plugins/Experimental/StructUtils/Source/StructUtils/Public/StructArrayView.h
+++ b/Engine/Plugins/Experimental/StructUtils/Source/StructUtils/Public/StructArrayView.h
@@ -44,15 +44,24 @@ struct FStructArrayView
 		check((InData != nullptr) || ((InScriptStruct == nullptr) && (InElementSize == 0) && (InCount == 0)));
 	}
 
-	/**
-	 * Construct a view with given data.
-	 */
+	FStructArrayView(const UScriptStruct& InScriptStruct, void* InData, const int32 InCount)
+		: DataPtr(InData)
+		, ScriptStruct(&InScriptStruct)
+		, ElementSize(InScriptStruct.GetStructureSize())
+		, ArrayNum(InCount)
+	{
+		check(InCount >= 0);
+		check(InData != nullptr || InCount == 0);
+	}
+
+	UE_DEPRECATED(5.3, "Use constructor that takes a reference instead")
 	FStructArrayView(const UScriptStruct* InScriptStruct, void* InData, const int32 InCount)
 		: DataPtr(InData)
 		, ScriptStruct(InScriptStruct)
 		, ElementSize(InScriptStruct->GetStructureSize())
 		, ArrayNum(InCount)
 	{
+		check(InScriptStruct);
 		check(InCount >= 0);
 		check(InData != nullptr || InCount == 0);
 	}
@@ -421,13 +430,10 @@ struct FConstStructArrayView
 		check((InData != nullptr) || ((InScriptStruct == nullptr) && (InElementSize == 0) && (InCount == 0)));
 	}
 
-	/**
-	 * Construct a view with given data.
-	 */
-	FConstStructArrayView(const UScriptStruct* InScriptStruct, const void* InData, const int32 InCount)
+	FConstStructArrayView(const UScriptStruct& InScriptStruct, const void* InData, const int32 InCount)
 		: DataPtr(InData)
-		, ScriptStruct(InScriptStruct)
-		, ElementSize(InScriptStruct->GetStructureSize())
+		, ScriptStruct(&InScriptStruct)
+		, ElementSize(InScriptStruct.GetStructureSize())
 		, ArrayNum(InCount)
 	{
 		check(InCount >= 0);
