Minor Horde warning fix

diff --git a/Engine/Plugins/Experimental/StructUtils/Source/StructUtilsTestSuite/Private/StructSharedStructTest.cpp b/Engine/Plugins/Experimental/StructUtils/Source/StructUtilsTestSuite/Private/StructSharedStructTest.cpp
--- a/Engine/Plugins/Experimental/StructUtils/Source/StructUtilsTestSuite/Private/StructSharedStructTest.cpp
+++ b/Engine/Plugins/Experimental/StructUtils/Source/StructUtilsTestSuite/Private/StructSharedStructTest.cpp
@@ -6,7 +6,7 @@
 
 #define LOCTEXT_NAMESPACE "StructUtilsTests"
 
-PRAGMA_DISABLE_OPTIMIZATION
+UE_DISABLE_OPTIMIZATION_SHIP
 
 namespace FStructSharedStructTest
 {
@@ -226,6 +226,6 @@ namespace FStructSharedStructTest
 }
 
 
-PRAGMA_ENABLE_OPTIMIZATION
+UE_ENABLE_OPTIMIZATION_SHIP
 
 #undef LOCTEXT_NAMESPACE
