Fix for minor Horde error

diff --git a/Engine/Plugins/Experimental/StructUtils/Source/StructUtils/Private/StructUtilsTypes.cpp b/Engine/Plugins/Experimental/StructUtils/Source/StructUtils/Private/StructUtilsTypes.cpp
--- a/Engine/Plugins/Experimental/StructUtils/Source/StructUtils/Private/StructUtilsTypes.cpp
+++ b/Engine/Plugins/Experimental/StructUtils/Source/StructUtils/Private/StructUtilsTypes.cpp
@@ -29,12 +29,12 @@ namespace UE::StructUtils
 		return 0;
 	}
 
-	STRUCTUTILS_API uint32 GetStructCrc32(const FStructView StructView, const uint32 CRC /*= 0*/)
+	STRUCTUTILS_API uint32 GetStructCrc32(const FStructView& StructView, const uint32 CRC /*= 0*/)
 	{
 		return GetStructCrc32Helper(StructView, CRC);
 	}
 
-	STRUCTUTILS_API uint32 GetStructCrc32(const FConstStructView StructView, const uint32 CRC /*= 0*/)
+	STRUCTUTILS_API uint32 GetStructCrc32(const FConstStructView& StructView, const uint32 CRC /*= 0*/)
 	{
 		return GetStructCrc32Helper(StructView, CRC);
 	}
diff --git a/Engine/Plugins/Experimental/StructUtils/Source/StructUtils/Public/StructUtilsTypes.h b/Engine/Plugins/Experimental/StructUtils/Source/StructUtils/Public/StructUtilsTypes.h
--- a/Engine/Plugins/Experimental/StructUtils/Source/StructUtils/Public/StructUtilsTypes.h
+++ b/Engine/Plugins/Experimental/StructUtils/Source/StructUtils/Public/StructUtilsTypes.h
@@ -2,7 +2,6 @@
 
 #pragma once
 
-#include "StructView.h"
 #include <type_traits>
 
 #ifndef WITH_STRUCTUTILS_DEBUG
@@ -19,8 +18,8 @@ class FReferenceCollector;
 namespace UE::StructUtils
 {
 	extern STRUCTUTILS_API uint32 GetStructCrc32(const UScriptStruct& ScriptStruct, const uint8* StructMemory, const uint32 CRC = 0);
-	extern STRUCTUTILS_API uint32 GetStructCrc32(const FStructView StructView, const uint32 CRC = 0);
-	extern STRUCTUTILS_API uint32 GetStructCrc32(const FConstStructView StructView, const uint32 CRC = 0);
+	extern STRUCTUTILS_API uint32 GetStructCrc32(const FStructView& StructView, const uint32 CRC = 0);
+	extern STRUCTUTILS_API uint32 GetStructCrc32(const FConstStructView& StructView, const uint32 CRC = 0);
 	extern STRUCTUTILS_API uint32 GetStructCrc32(const FSharedStruct& SharedView, const uint32 CRC = 0);
 	extern STRUCTUTILS_API uint32 GetStructCrc32(const FConstSharedStruct& SharedView, const uint32 CRC = 0);
 
