StructUtils: fix compilation in DebugGame

diff --git a/Engine/Plugins/Experimental/StructUtils/Source/StructUtils/Public/StructView.h b/Engine/Plugins/Experimental/StructUtils/Source/StructUtils/Public/StructView.h
--- a/Engine/Plugins/Experimental/StructUtils/Source/StructUtils/Public/StructView.h
+++ b/Engine/Plugins/Experimental/StructUtils/Source/StructUtils/Public/StructView.h
@@ -21,7 +21,7 @@ struct FConstSharedStruct;
  * data can be mutated. Use FConstStructView to prevent mutation of the actual struct data.
  * See FConstStructView for examples.
  */
-struct STRUCTUTILS_API FStructView
+struct FStructView
 {
 public:
 
@@ -36,7 +36,7 @@ public:
 		: FStructView(InstancedStruct.GetScriptStruct(), InstancedStruct.GetMutableMemory())
 	{}
 
-	FStructView(FSharedStruct& SharedStruct);
+	STRUCTUTILS_API FStructView(FSharedStruct& SharedStruct);
 
 	/** Creates a new FStructView from the templated struct. Note its not safe to make InStruct const ref as the original object may have been declared const */
 	template<typename T>
@@ -144,7 +144,7 @@ protected:
  * e.g. FStructView A; const FStructView B; A = B; // compiles as the struct B is pointing to can't be made to point at something else but A isn't cosnt.
  * e.g. const FStructView A; FStructView B; A = B; // doesn't compile as attempting to make const view point at something else
  */
-struct STRUCTUTILS_API FConstStructView
+struct FConstStructView
 {
 public:
 
@@ -159,9 +159,9 @@ public:
 		: FConstStructView(InstancedStruct.GetScriptStruct(), InstancedStruct.GetMemory())
 	{}
 
-	FConstStructView(const FSharedStruct SharedStruct);
+	STRUCTUTILS_API FConstStructView(const FSharedStruct SharedStruct);
 
-	FConstStructView(const FConstSharedStruct SharedStruct);
+	STRUCTUTILS_API FConstStructView(const FConstSharedStruct SharedStruct);
 
 	FConstStructView(const FStructView StructView)
 		: FConstStructView(StructView.GetScriptStruct(), StructView.GetMemory())
