Fix for InstancedStruct.h not being included from StructTypeBitSet.h

diff --git a/Engine/Plugins/Experimental/StructUtils/Source/StructUtils/Public/StructTypeBitSet.h b/Engine/Plugins/Experimental/StructUtils/Source/StructUtils/Public/StructTypeBitSet.h
--- a/Engine/Plugins/Experimental/StructUtils/Source/StructUtils/Public/StructTypeBitSet.h
+++ b/Engine/Plugins/Experimental/StructUtils/Source/StructUtils/Public/StructTypeBitSet.h
@@ -3,6 +3,7 @@
 #pragma once
 
 #include "StructUtilsTypes.h"
+#include "InstancedStruct.h"
 
 
 /**
