Remove StructView.cpp and STRUCTUTILS_API macros from SharedView structs.

diff --git a/Engine/Plugins/Experimental/StructUtils/Source/StructUtils/Private/StructView.cpp b/Engine/Plugins/Experimental/StructUtils/Source/StructUtils/Private/StructView.cpp
deleted file mode 100644
diff --git a/Engine/Plugins/Experimental/StructUtils/Source/StructUtils/Public/StructView.h b/Engine/Plugins/Experimental/StructUtils/Source/StructUtils/Public/StructView.h
--- a/Engine/Plugins/Experimental/StructUtils/Source/StructUtils/Public/StructView.h
+++ b/Engine/Plugins/Experimental/StructUtils/Source/StructUtils/Public/StructView.h
@@ -3,11 +3,9 @@
 #pragma once
 
 #include "InstancedStruct.h"
+#include "SharedStruct.h"
 #include "StructUtils.h"
 
-struct FSharedStruct;
-struct FConstSharedStruct;
-
 ///////////////////////////////////////////////////////////////// FStructView /////////////////////////////////////////////////////////////////
 
 /**
@@ -17,7 +15,7 @@ struct FConstSharedStruct;
  * E.g. instead of passing ref or pointer to a FInstancedStruct, you should use FConstStructView or FStructView to pass around a view to the contents.
  * FStructView is passed by value.
  * FStructView is similar to FStructOnScope, but FStructView is a view only (FStructOnScope can either own the memory or be a view)
- * const FStructView prevents the struct from pointing at a different isntance of a struct. However the actual struct 
+ * const FStructView prevents the struct from pointing at a different instance of a struct. However the actual struct 
  * data can be mutated. Use FConstStructView to prevent mutation of the actual struct data.
  * See FConstStructView for examples.
  */
@@ -36,7 +34,9 @@ public:
 		: FStructView(InstancedStruct.GetScriptStruct(), InstancedStruct.GetMutableMemory())
 	{}
 
-	STRUCTUTILS_API FStructView(FSharedStruct& SharedStruct);
+	FStructView(const FSharedStruct& SharedStruct)
+		: FStructView(SharedStruct.GetScriptStruct(), SharedStruct.GetMemory())
+	{}
 
 	/** Creates a new FStructView from the templated struct. Note its not safe to make InStruct const ref as the original object may have been declared const */
 	template<typename T>
@@ -159,9 +159,13 @@ public:
 		: FConstStructView(InstancedStruct.GetScriptStruct(), InstancedStruct.GetMemory())
 	{}
 
-	STRUCTUTILS_API FConstStructView(const FSharedStruct SharedStruct);
+	FConstStructView(const FSharedStruct& SharedStruct)
+		: FConstStructView(SharedStruct.GetScriptStruct(), SharedStruct.GetMemory())
+	{}
 
-	STRUCTUTILS_API FConstStructView(const FConstSharedStruct SharedStruct);
+	FConstStructView(const FConstSharedStruct& SharedStruct)
+		: FConstStructView(SharedStruct.GetScriptStruct(), SharedStruct.GetMemory())
+	{}
 
 	FConstStructView(const FStructView StructView)
 		: FConstStructView(StructView.GetScriptStruct(), StructView.GetMemory())
@@ -254,7 +258,3 @@ protected:
 	const UScriptStruct* ScriptStruct = nullptr;
 	const uint8* StructMemory = nullptr;
 };
-
-#if UE_ENABLE_INCLUDE_ORDER_DEPRECATED_IN_5_3
-#include "SharedStruct.h"
-#endif
\ No newline at end of file
