Minor fix for static analysis warning

diff --git a/Engine/Plugins/Experimental/StructUtils/Source/StructUtils/Public/StructUtils.h b/Engine/Plugins/Experimental/StructUtils/Source/StructUtils/Public/StructUtils.h
--- a/Engine/Plugins/Experimental/StructUtils/Source/StructUtils/Public/StructUtils.h
+++ b/Engine/Plugins/Experimental/StructUtils/Source/StructUtils/Public/StructUtils.h
@@ -11,9 +11,9 @@ namespace UE::StructUtils
 	void CheckStructType()
 	{
 		static_assert(!TIsDerivedFrom<T, struct FInstancedStruct>::IsDerived &&
-					  !TIsDerivedFrom<T, struct FConstStructView>::IsDerived &&
-					  !TIsDerivedFrom<T, struct FStructView>::IsDerived &&
 					  !TIsDerivedFrom<T, struct FStructView>::IsDerived &&
+					  !TIsDerivedFrom<T, struct FConstStructView>::IsDerived &&
+					  !TIsDerivedFrom<T, struct FSharedStruct>::IsDerived &&
 					  !TIsDerivedFrom<T, struct FConstSharedStruct>::IsDerived, "It does not make sense to create wrapper over these types.");
 	}
 
