BitSet serialization

diff --git a/Engine/Plugins/Experimental/StructUtils/Source/StructUtils/Private/StructTypeBitSet.cpp b/Engine/Plugins/Experimental/StructUtils/Source/StructUtils/Private/StructTypeBitSet.cpp
new file mode 100644
--- /dev/null
+++ b/Engine/Plugins/Experimental/StructUtils/Source/StructUtils/Private/StructTypeBitSet.cpp
@@ -0,0 +1,230 @@
+// Copyright Epic Games, Inc. All Rights Reserved.
+
+#include "StructTypeBitSet.h"
+#include "Serialization/Archive.h"
+
+
+DEFINE_LOG_CATEGORY_STATIC(LogStructUtils, Warning, All);
+
+namespace FBitSetSerializationDictionary
+{
+	TMap<uint32 /*CurrentStoredTypesHash*/, const FStructTracker* /*Tracker*/> TrackerMap;
+	TMap<uint32 /*NotUpToDateStoredTypesHash*/, TArray<int32> /*Mapping*/> BitMappings;
+
+	void RegisterHash(const uint32 SerializationHash, const FStructTracker& Tracker)
+	{
+		if (const FStructTracker** StoredTracer = TrackerMap.Find(SerializationHash))
+		{
+			ensureMsgf(*StoredTracer == nullptr || *StoredTracer == &Tracker
+				, TEXT("Hash conflict when registering a FStructTracker instance"));
+		}
+		TrackerMap.Add(SerializationHash, &Tracker);
+	}
+
+	const FStructTracker** GetTracker(const uint32 SerializationHash)
+	{
+		return TrackerMap.Find(SerializationHash);
+	}
+
+	TArray<int32>& GetOrCreateBitMapping(const uint32 SerializationHash)
+	{
+		return BitMappings.FindOrAdd(SerializationHash);
+	}
+
+	void UnregisterTracker(const FStructTracker& Tracker)
+	{
+		for (auto It = TrackerMap.CreateIterator(); It; ++It)
+		{
+			if (It->Value == &Tracker)
+			{
+				// note that we're not bailing out once the entry is removes, since
+				// Trackers are being registered for all the hashes calculated as 
+				// the new types to track get added
+				It.RemoveCurrent();
+			}
+		}
+	}
+}
+
+
+//-----------------------------------------------------------------------------
+// FStructTracker
+//-----------------------------------------------------------------------------
+FStructTracker::FStructTracker(const FBaseStructGetter& InBaseStructGetter)
+	: BaseStructGetter(InBaseStructGetter)
+{
+}
+
+FStructTracker::~FStructTracker()
+{
+	FBitSetSerializationDictionary::UnregisterTracker(*this);
+}
+
+int32 FStructTracker::FindOrAddStructTypeIndex(const UStruct& InStructType)
+{
+	// Get existing index...
+	const uint32 Hash = PointerHash(&InStructType);
+	FSetElementId ElementId = StructTypeToIndexSet.FindIdByHash(Hash, Hash);
+
+	if (!ElementId.IsValidId())
+	{
+		checkSlow(InStructType.IsChildOf(BaseStructGetter()));
+
+		// .. or create a new one
+		ElementId = StructTypeToIndexSet.AddByHash(Hash, Hash);
+		checkSlow(ElementId.IsValidId());
+
+		const int32 NewIndex = ElementId.AsInteger();
+
+		checkSlow(StructTypesList.Num() == NewIndex);
+		StructTypesList.Add(&InStructType);
+
+		// first-time SerializationHash initialization
+		if (SerializationHash == 0)
+		{
+			ensure(StructTypesList.Num() == 1);
+			UStruct* BaseType = BaseStructGetter();
+			CA_ASSUME(BaseType);
+			SerializationHash = GetTypeHash(BaseType->GetFullName());
+		}
+		SerializationHash = HashCombine(SerializationHash, GetTypeHash(InStructType.GetFullName()));
+
+		// it's worth pointing out that we're registering a given tracker for all the hashes created along the way.
+		// This will help with loading bitsets from serialized data.
+		FBitSetSerializationDictionary::RegisterHash(SerializationHash, *this);
+
+#if WITH_STRUCTUTILS_DEBUG
+		DebugStructTypeNamesList.Add(InStructType.GetFName());
+		ensure(StructTypeToIndexSet.Num() == DebugStructTypeNamesList.Num());
+#endif // WITH_STRUCTUTILS_DEBUG
+	}
+
+	return ElementId.AsInteger();
+}
+
+void FStructTracker::Serialize(FArchive& Ar, FStructTypeBitSet::FBitSetContainer& StructTypesBitArray)
+{	
+	enum class EVersion : uint8
+	{
+		Initial,
+		Last,
+		Current = Last - 1
+	};
+
+	uint8 Version = uint8(EVersion::Current);
+	Ar << Version;
+
+	if (Ar.IsSaving())
+	{
+		Ar << SerializationHash;
+		
+		const int64 SizeOffset = Ar.Tell(); 
+		int32 SerialSize = 0;
+		Ar << SerialSize;
+
+		// Serialized memory
+		const int64 InitialDataOffset = Ar.Tell();
+
+		// store information on the base type, so that we can verify we're trying to load the right data later on
+		FTopLevelAssetPath BaseStructPath(BaseStructGetter());
+		Ar << BaseStructPath;
+
+		int32 StructTypesListNum = StructTypesList.Num();
+		Ar << StructTypesListNum;
+	
+		for (const TWeakObjectPtr<const UStruct>& StructType : StructTypesList)
+		{
+			FTopLevelAssetPath StructPath(StructType.Get());
+			Ar << StructPath;
+		}
+
+		const int64 FinalDataOffset = Ar.Tell();
+
+		// Size of the serialized memory
+		Ar.Seek(SizeOffset);	
+		SerialSize = (int32)(FinalDataOffset - InitialDataOffset);
+		Ar << SerialSize;
+		// switch back to the end position
+		Ar.Seek(FinalDataOffset);
+
+		// serializing the actual bits
+		Ar << StructTypesBitArray;
+	}
+	else if (Ar.IsLoading())
+	{
+		uint32 LoadedSerializationHash = 0;
+		Ar << LoadedSerializationHash;
+		
+		int32 SerialSize = 0;
+		Ar << SerialSize;
+
+		const FStructTracker** KnownTracker = FBitSetSerializationDictionary::GetTracker(LoadedSerializationHash);
+
+		if (KnownTracker && *KnownTracker == this)
+		{
+			// can skip the whole Tracker saved data
+			Ar.Seek(Ar.Tell() + SerialSize);
+			Ar << StructTypesBitArray;
+		}
+		else
+		{
+			FTopLevelAssetPath BaseStructPath;
+			Ar << BaseStructPath;
+			// first, verify that the base type of the tracker matches - otherwise we can accidentally storing 
+			// the wrong types creating the Mapping (via FindOrAddStructTypeIndex below)
+			UStruct* SerializedBaseStructType = FindObject<UStruct>(BaseStructPath);
+
+			if (SerializedBaseStructType == BaseStructGetter())
+			{
+				// create a translator 
+				TArray<int32>& BitMapping = FBitSetSerializationDictionary::GetOrCreateBitMapping(LoadedSerializationHash);
+				if (BitMapping.Num() == 0)
+				{
+					int32 StructTypesListNum = 0;
+					Ar << StructTypesListNum;
+
+					BitMapping.AddUninitialized(StructTypesListNum);
+
+					for (int32 Index = 0; Index < StructTypesListNum; ++Index)
+					{
+						FTopLevelAssetPath TypePath;
+						Ar << TypePath;
+
+						if (UStruct* StructType = FindObject<UStruct>(TypePath))
+						{
+							checkSlow(StructType->IsChildOf(BaseStructGetter()));
+							BitMapping[Index] = FindOrAddStructTypeIndex(*StructType);
+						}
+						else
+						{
+							BitMapping[Index] = INDEX_NONE;
+						}
+					}
+				}
+				else
+				{
+					Ar.Seek(Ar.Tell() + SerialSize);
+				}
+
+				// it's where we read in data saved with a different order
+				TBitArray<> TempStructTypesBitArray;
+				Ar << TempStructTypesBitArray;
+			
+				StructTypesBitArray.Init(false, StructTypeToIndexSet.Num());
+				for (TBitArray<>::FConstIterator It(TempStructTypesBitArray); It; ++It)
+				{
+					if (It.GetValue())
+					{
+						const int32 TranslatedIndex = BitMapping[It.GetIndex()];
+						StructTypesBitArray.AddAtIndex(TranslatedIndex);
+					}
+				}
+			}
+			else
+			{
+				UE_LOG(LogStructUtils, Error, TEXT("Trying to load mismatching BitSet data. Current base class: %s, read base class: %s")
+					, *GetNameSafe(BaseStructGetter()), *BaseStructPath.ToString());
+			}
+		}
+	}
+}
diff --git a/Engine/Plugins/Experimental/StructUtils/Source/StructUtils/Public/StructTypeBitSet.h b/Engine/Plugins/Experimental/StructUtils/Source/StructUtils/Public/StructTypeBitSet.h
--- a/Engine/Plugins/Experimental/StructUtils/Source/StructUtils/Public/StructTypeBitSet.h
+++ b/Engine/Plugins/Experimental/StructUtils/Source/StructUtils/Public/StructTypeBitSet.h
@@ -2,10 +2,160 @@
 
 #pragma once
 
+#include "Containers/BitArray.h"
 #include "StructUtilsTypes.h"
 #include "InstancedStruct.h"
 
 
+class FArchive;
+
+namespace FStructTypeBitSet
+{
+struct FBitSetContainer : TBitArray<>
+{
+	FBitSetContainer() = default;
+	FBitSetContainer(const TBitArray<>& Source) : TBitArray<>(Source)
+	{}
+
+	FBitSetContainer& operator=(const TBitArray<>& Other)
+	{
+		*((TBitArray<>*)this) = Other;
+		return *this;
+	}
+
+	void SetAll(const bool bValue, const int32 Count)
+	{
+		Init(bValue, Count);
+	}
+
+	FORCEINLINE bool HasAll(const TBitArray<>& Other) const
+	{
+		FConstWordIterator ThisIterator(*this);
+		FConstWordIterator OtherIterator(Other);
+
+		while (ThisIterator || OtherIterator)
+		{
+			const uint32 A = ThisIterator ? ThisIterator.GetWord() : 0;
+			const uint32 B = OtherIterator ? OtherIterator.GetWord() : 0;
+			if ((A & B) != B)
+			{
+				return false;
+			}
+
+			++ThisIterator;
+			++OtherIterator;
+		}
+
+		return true;
+	}
+
+	FORCEINLINE bool HasAny(const TBitArray<>& Other) const
+	{
+		FConstWordIterator ThisIterator(*this);
+		FConstWordIterator OtherIterator(Other);
+
+		while (ThisIterator || OtherIterator)
+		{
+			const uint32 A = ThisIterator ? ThisIterator.GetWord() : 0;
+			const uint32 B = OtherIterator ? OtherIterator.GetWord() : 0;
+			if ((A & B) != 0)
+			{
+				return true;
+			}
+
+			++ThisIterator;
+			++OtherIterator;
+		}
+
+		return false;
+	}
+
+	FORCEINLINE bool IsEmpty() const
+	{
+		FConstWordIterator Iterator(*this);
+
+		while (Iterator && Iterator.GetWord() == 0)
+		{
+			++Iterator;
+		}
+
+		return !Iterator;
+	}
+
+	FORCEINLINE void operator-=(const TBitArray<>& Other)
+	{
+		FWordIterator ThisIterator(*this);
+		FConstWordIterator OtherIterator(Other);
+
+		while (ThisIterator && OtherIterator)
+		{
+			ThisIterator.SetWord(ThisIterator.GetWord() & ~OtherIterator.GetWord());
+
+			++ThisIterator;
+			++OtherIterator;
+		}
+	}
+
+	FORCEINLINE friend uint32 GetTypeHash(const FBitSetContainer& Instance)
+	{
+		FConstWordIterator Iterator(Instance);
+		uint32 Hash = 0;
+		uint32 TrailingZeroHash = 0;
+		while (Iterator)
+		{
+			const uint32 Word = Iterator.GetWord();
+			if (Word)
+			{
+				Hash = HashCombine(TrailingZeroHash ? TrailingZeroHash : Hash, Word);
+				TrailingZeroHash = 0;
+			}
+			else // potentially a trailing 0-word
+			{
+				TrailingZeroHash = HashCombine(TrailingZeroHash ? TrailingZeroHash : Hash, Word);
+			}
+			++Iterator;
+		}
+		return Hash;
+	}
+
+	void AddAtIndex(const int32 Index)
+	{
+		PadToNum(Index + 1, false);
+		SetBitNoCheck(Index, true);
+	}
+
+	void RemoveAtIndex(const int32 Index)
+	{
+		check(Index >= 0);
+		if (Index < Num())
+		{
+			SetBitNoCheck(Index, false);
+		}
+		// else, it's already not present
+	}
+
+	bool Contains(const int32 Index) const
+	{
+		check(Index >= 0);
+		return (Index < Num()) && (*this)[Index];
+	}
+
+protected:
+	/**
+	 * duplication of TBitArray::SetBitNoCheck needed since it's private but it's the performant way of setting bits
+	 * when we know the index is valid.
+	 * @todo ask Core team about exposing that
+	 */
+	void SetBitNoCheck(const int32 Index, const bool Value)
+	{
+		uint32& Word = GetData()[Index / NumBitsPerDWORD];
+		const uint32 BitOffset = (Index % NumBitsPerDWORD);
+		Word = (Word & ~(1 << BitOffset)) | (((uint32)Value) << BitOffset);
+	}
+};
+} // namespace FStructTypeBitSet
+
+
 /**
  * The TStructTypeBitSet holds information on "existence" of subtypes of a given UStruct. The information on 
  * available child-structs is gathered lazily - the internal FStructTracker is assigning a given type a new index the
@@ -28,34 +178,20 @@
  *	DEFINE_TYPEBITSET(FMyFooBarBitSet);
  * 
  */
-struct FStructTracker
+struct STRUCTUTILS_API FStructTracker
 {
-	int32 FindOrAddStructTypeIndex(const UStruct& InStructType)
-	{
-		// Get existing index...
-		const uint32 Hash = PointerHash(&InStructType);
-		FSetElementId ElementId = StructTypeToIndexSet.FindIdByHash(Hash, Hash);
+	using FBaseStructGetter = TFunction<UStruct*()>;
 
-		if (!ElementId.IsValidId())
-		{
-			// .. or create new one
-			ElementId = StructTypeToIndexSet.AddByHash(Hash, Hash);
-			checkSlow(ElementId.IsValidId());
+	/** 
+	 * The input parameter is a function that fetches the UStruct representing the base class for all the stored structs.
+	 * We're unable to get the UStruct parameter directly since FStructTracker instances are being created during
+	 * module loading (via DEFINE_TYPEBITSET macro) and the ::StaticStruct call fail at that point for types defined
+	 * in the same module.
+	 */
+	explicit FStructTracker(const FBaseStructGetter& InBaseStructGetter);
+	~FStructTracker();
 
-			checkSlow(StructTypesList.Num() == ElementId.AsInteger());
-			StructTypesList.Add(&InStructType);
-		}
-		const int32 Index = ElementId.AsInteger();
-
-#if WITH_STRUCTUTILS_DEBUG
-		if (Index == DebugStructTypeNamesList.Num())
-		{
-			DebugStructTypeNamesList.Add(InStructType.GetFName());
-			ensure(StructTypeToIndexSet.Num() == DebugStructTypeNamesList.Num());
-		}
-#endif // WITH_STRUCTUTILS_DEBUG
-		return Index;
-	}
+	int32 FindOrAddStructTypeIndex(const UStruct& InStructType);
 
 	const UStruct* GetStructType(const int32 StructTypeIndex) const
 	{
@@ -64,6 +200,8 @@ struct FStructTracker
 
 	int32 Num() const {	return StructTypeToIndexSet.Num(); }
 
+	void Serialize(FArchive& Ar, FStructTypeBitSet::FBitSetContainer& StructTypesBitArray);
+
 #if WITH_STRUCTUTILS_DEBUG
 	/**
 	* @return index identifying given tag or INDEX_NONE if it has never been used/seen before.
@@ -90,158 +228,16 @@ struct FStructTracker
 
 	TSet<uint32> StructTypeToIndexSet;
 	TArray<TWeakObjectPtr<const UStruct>, TInlineAllocator<64>> StructTypesList;
+	uint32 SerializationHash = 0;
+	const FBaseStructGetter BaseStructGetter;
 };
 
+
 template<typename TBaseStruct, typename TStructTrackerWrapper, typename TUStructType = UScriptStruct>
 struct TStructTypeBitSet
 {
 	using FStructTrackerWrapper = TStructTrackerWrapper;
 
-private:
-	struct FBitArrayExt : TBitArray<>
-	{
-		FBitArrayExt() = default;
-		FBitArrayExt(const TBitArray<>& Source) : TBitArray<>(Source)
-		{}
-
-		FBitArrayExt& operator=(const TBitArray<>& Other)
-		{
-			*((TBitArray<>*)this) = Other;
-			return *this;
-		}
-
-		void SetAll(const bool bValue)
-		{
-			Init(bValue, TStructTrackerWrapper::StructTracker.Num());
-		}
-
-		FORCEINLINE bool HasAll(const TBitArray<>& Other) const
-		{
-			FConstWordIterator ThisIterator(*this);
-			FConstWordIterator OtherIterator(Other);
-
-			while (ThisIterator || OtherIterator)
-			{
-				const uint32 A = ThisIterator ? ThisIterator.GetWord() : 0;
-				const uint32 B = OtherIterator ? OtherIterator.GetWord() : 0;
-				if ((A & B) != B)
-				{
-					return false;
-				}
-
-				++ThisIterator;
-				++OtherIterator;
-			}
-
-			return true;
-		}
-
-		FORCEINLINE bool HasAny(const TBitArray<>& Other) const
-		{
-			FConstWordIterator ThisIterator(*this);
-			FConstWordIterator OtherIterator(Other);
-
-			while (ThisIterator || OtherIterator)
-			{
-				const uint32 A = ThisIterator ? ThisIterator.GetWord() : 0;
-				const uint32 B = OtherIterator ? OtherIterator.GetWord() : 0;
-				if ((A & B) != 0)
-				{
-					return true;
-				}
-
-				++ThisIterator;
-				++OtherIterator;
-			}
-
-			return false;
-		}
-
-		FORCEINLINE bool IsEmpty() const
-		{
-			FConstWordIterator Iterator(*this);
-
-			while (Iterator && Iterator.GetWord() == 0)
-			{
-				++Iterator;
-			}
-
-			return !Iterator;
-		}
-
-		FORCEINLINE void operator-=(const TBitArray<>& Other)
-		{
-			FWordIterator ThisIterator(*this);
-			FConstWordIterator OtherIterator(Other);
-
-			while (ThisIterator && OtherIterator)
-			{
-				ThisIterator.SetWord(ThisIterator.GetWord() & ~OtherIterator.GetWord());
-
-				++ThisIterator;
-				++OtherIterator;
-			}
-		}
-
-		FORCEINLINE friend uint32 GetTypeHash(const FBitArrayExt& Instance)
-		{
-			FConstWordIterator Iterator(Instance);
-			uint32 Hash = 0;
-			uint32 TrailingZeroHash = 0;
-			while (Iterator)
-			{
-				const uint32 Word = Iterator.GetWord();
-				if (Word)
-				{
-					Hash = HashCombine(TrailingZeroHash ? TrailingZeroHash : Hash, Word);
-					TrailingZeroHash = 0;
-				}
-				else // potentially a trailing 0-word
-				{
-					TrailingZeroHash = HashCombine(TrailingZeroHash ? TrailingZeroHash : Hash, Word);
-				}
-				++Iterator;
-			}
-			return Hash;
-		}
-
-		void AddAtIndex(const int32 Index)
-		{
-			PadToNum(Index + 1, false);
-			SetBitNoCheck(Index, true);
-		}
-
-		void RemoveAtIndex(const int32 Index)
-		{
-			check(Index >= 0);
-			if (Index < Num())
-			{
-				SetBitNoCheck(Index, false);
-			}
-			// else, it's already not present
-		}
-
-		bool Contains(const int32 Index) const
-		{
-			check(Index >= 0);
-			return (Index < Num()) && (*this)[Index];
-		}
-
-	protected:
-		/**
-		 * duplication of TBitArray::SetBitNoCheck needed since it's private but it's the performant way of setting bits
-		 * when we know the index is valid.
-		 * @todo ask Core team about exposing that
-		 */
-		void SetBitNoCheck(const int32 Index, const bool Value)
-		{
-			uint32& Word = GetData()[Index / NumBitsPerDWORD];
-			const uint32 BitOffset = (Index % NumBitsPerDWORD);
-			Word = (Word & ~(1 << BitOffset)) | (((uint32)Value) << BitOffset);
-		}
-	};
-
-public:
 	TStructTypeBitSet() = default;
 
 	explicit TStructTypeBitSet(const TUStructType& StructType)
@@ -288,7 +284,7 @@ public:
 	struct FIndexIterator
 	{
 	public:
-		explicit FIndexIterator(const FBitArrayExt& BitArray, const bool bInValueToCheck = true)
+		explicit FIndexIterator(const FStructTypeBitSet::FBitSetContainer& BitArray, const bool bInValueToCheck = true)
 			: It(BitArray), bValueToCheck(bInValueToCheck)
 		{
 			if (It && It.GetValue() != bInValueToCheck)
@@ -410,7 +406,7 @@ public:
 
 	void SetAll(const bool bValue = true)
 	{
-		StructTypesBitArray.SetAll(bValue);
+		StructTypesBitArray.SetAll(bValue, TStructTrackerWrapper::StructTracker.Num());
 	}
 
 	void Add(const TUStructType& InStructType)
@@ -642,8 +638,19 @@ public:
 		return HashCombine(StoredTypeHash, BitArrayHash);
 	}
 
+	void Serialize(FArchive& Ar)
+	{
+		TStructTrackerWrapper::StructTracker.Serialize(Ar, StructTypesBitArray);
+	}
+
+	FORCEINLINE friend FArchive& operator<<(FArchive& Ar, TStructTypeBitSet& Instance)
+	{
+		Instance.Serialize(Ar);
+		return Ar;
+	}
+
 private:
-	FBitArrayExt StructTypesBitArray;
+	FStructTypeBitSet::FBitSetContainer StructTypesBitArray;
 };
 
 /** 
@@ -652,12 +659,13 @@ private:
  * way we could come up with to solve it. Thankfully the user doesn't need to even know about this class's existence
  * as long as they are using the macros below.
  */
-#define _DECLARE_TYPEBITSET_IMPL(EXPORTED_API, ContainerTypeName, BaseType, BaseUStructType) struct ContainerTypeName##StructTracker \
+#define _DECLARE_TYPEBITSET_IMPL(EXPORTED_API, ContainerTypeName, BaseType, BaseUStructType) struct ContainerTypeName##StructTrackerWrapper \
 	{ \
+		using FBaseStructType = BaseType; \
 		EXPORTED_API static FStructTracker StructTracker; \
 	}; \
-	template struct EXPORTED_API TStructTypeBitSet<BaseType, ContainerTypeName##StructTracker, BaseUStructType>; \
-	using ContainerTypeName = TStructTypeBitSet<BaseType, ContainerTypeName##StructTracker, BaseUStructType>
+	template struct EXPORTED_API TStructTypeBitSet<BaseType, ContainerTypeName##StructTrackerWrapper, BaseUStructType>; \
+	using ContainerTypeName = TStructTypeBitSet<BaseType, ContainerTypeName##StructTrackerWrapper, BaseUStructType>
 
 #define DECLARE_STRUCTTYPEBITSET_EXPORTED(EXPORTED_API, ContainerTypeName, BaseStructType) _DECLARE_TYPEBITSET_IMPL(EXPORTED_API, ContainerTypeName, BaseStructType, UScriptStruct)
 #define DECLARE_STRUCTTYPEBITSET(ContainerTypeName, BaseStructType) _DECLARE_TYPEBITSET_IMPL(, ContainerTypeName, BaseStructType, UScriptStruct)
@@ -665,4 +673,4 @@ private:
 #define DECLARE_CLASSTYPEBITSET(ContainerTypeName, BaseStructType) _DECLARE_TYPEBITSET_IMPL(, ContainerTypeName, BaseStructType, UClass)
 
 #define DEFINE_TYPEBITSET(ContainerTypeName) \
-	FStructTracker ContainerTypeName##StructTracker::StructTracker;
+	FStructTracker ContainerTypeName##StructTrackerWrapper::StructTracker([](){ return UE::StructUtils::GetAsUStruct<ContainerTypeName##StructTrackerWrapper::FBaseStructType>();});
diff --git a/Engine/Plugins/Experimental/StructUtils/Source/StructUtilsTestSuite/Private/ScriptStructTypeBitSetTest.cpp b/Engine/Plugins/Experimental/StructUtils/Source/StructUtilsTestSuite/Private/ScriptStructTypeBitSetTest.cpp
--- a/Engine/Plugins/Experimental/StructUtils/Source/StructUtilsTestSuite/Private/ScriptStructTypeBitSetTest.cpp
+++ b/Engine/Plugins/Experimental/StructUtils/Source/StructUtilsTestSuite/Private/ScriptStructTypeBitSetTest.cpp
@@ -17,7 +17,7 @@ struct FTestStructBaseBitSetStructTracker
 {
 	static FStructTracker StructTracker;
 };
-FStructTracker FTestStructBaseBitSetStructTracker::StructTracker;
+FStructTracker FTestStructBaseBitSetStructTracker::StructTracker([](){return FTestStructSimpleBase::StaticStruct();});
 
 struct FTestStructBitSet : public TStructTypeBitSet<FTestStructSimple, FTestStructBaseBitSetStructTracker>
 {
diff --git a/Engine/Plugins/Runtime/MassEntity/Source/MassEntityTestSuite/Private/MassSerializationTest.cpp b/Engine/Plugins/Runtime/MassEntity/Source/MassEntityTestSuite/Private/MassSerializationTest.cpp
new file mode 100644
--- /dev/null
+++ b/Engine/Plugins/Runtime/MassEntity/Source/MassEntityTestSuite/Private/MassSerializationTest.cpp
@@ -0,0 +1,157 @@
+// Copyright Epic Games, Inc. All Rights Reserved.
+
+#include "MassEntityManager.h"
+#include "MassProcessingTypes.h"
+#include "MassEntityTestTypes.h"
+#include "MassExecutor.h"
+#include "MassEntityView.h"
+#include "Serialization/MemoryWriter.h"
+#include "Serialization/MemoryReader.h"
+
+#define LOCTEXT_NAMESPACE "MassTest"
+
+UE_DISABLE_OPTIMIZATION_SHIP
+
+struct FTestStructTracker : FStructTracker
+{
+	FTestStructTracker()
+		: FStructTracker([](){ return FMassFragment::StaticStruct(); })
+	{}
+
+	template<typename T>
+	int32 Add()
+	{
+		return StructTypesList.AddUnique(T::StaticStruct());
+	}
+
+	int32 Add(const UScriptStruct& Type)
+	{
+		return StructTypesList.AddUnique(&Type);
+	}
+};
+
+//----------------------------------------------------------------------//
+// tests 
+//----------------------------------------------------------------------//
+namespace FMassSerializationTest
+{
+
+//-----------------------------------------------------------------------------
+// BitSet serialization
+//-----------------------------------------------------------------------------
+struct FSerialization_BitSet : FExecutionTestBase
+{
+	virtual bool InstantTest() override
+	{
+		FMassFragmentBitSet BitSet;
+
+		BitSet.Add<FTestFragment_Float>();
+		BitSet.Add<FTestFragment_Bool>();
+		BitSet.Add<FTestFragment_Int>();
+
+		TArray<uint8> Data;
+		FMemoryWriter Writer(Data);
+
+		BitSet.Serialize(Writer);
+
+		FMemoryReader Reader(Data);
+		FMassFragmentBitSet NewBitSet;
+		NewBitSet.Serialize(Reader);
+
+		AITEST_TRUE("Saving and loading bitset of a given type should result in equivalent bitsets", BitSet.IsEquivalent(NewBitSet));
+		
+		return true;
+	}
+};
+IMPLEMENT_AI_INSTANT_TEST(FSerialization_BitSet, "System.Mass.Serialization.BitSet.Trivial");
+
+
+/*
+ * Here's what this test does:
+ * 1. creates a fake "old" tracker that knows about three types, and then creates a bit container having two of those 
+ *	set to true 
+ * 2. serializes this 2-out-of-3 bit container.
+ * 3. creates instance of the FMassFragmentBitSet type that hosts a tracker that knows about all the fragment structs,
+ *	including the ones used by "old fake" tracker.
+ * 4. reads the "outdated 2-out-of-3" data into an instance of FMassFragmentBitSet
+ * The point is that the order and indices of types in "old" tracker and the FMassFragmentBitSet's tracker are totally
+ * different, but loading patches it up.
+*/
+struct FSerialization_BitSetLoadOutdated : FExecutionTestBase
+{
+	virtual bool InstantTest() override
+	{
+		FTestStructTracker TestTracker;
+		const int32 FloatIndex = TestTracker.Add<FTestFragment_Float>();
+		TestTracker.Add<FTestFragment_Bool>();
+		const int32 IntIndex = TestTracker.Add<FTestFragment_Int>();
+		FStructTypeBitSet::FBitSetContainer TestBitArray;
+		TestBitArray.AddAtIndex(FloatIndex);
+		TestBitArray.AddAtIndex(IntIndex);
+
+		// create a bitset we're expected to get while loading TestBitArray's contents into FMassFragmentBitSet
+		FMassFragmentBitSet RegularBitSet;
+		RegularBitSet.Add<FTestFragment_Float>();
+		RegularBitSet.Add<FTestFragment_Int>();
+
+		// save
+		TArray<uint8> Data;
+		FMemoryWriter Writer(Data);
+		TestTracker.Serialize(Writer, TestBitArray);
+
+		// load
+		FMemoryReader Reader(Data);
+		FMassFragmentBitSet SerializedBitSet;
+		SerializedBitSet.Serialize(Reader);
+
+		AITEST_TRUE("Loading data serialized by the TestStructTracker should produce in a BitSet equivalent to the one explicitly created"
+			, RegularBitSet.IsEquivalent(SerializedBitSet));
+
+		return true;
+	}
+};
+IMPLEMENT_AI_INSTANT_TEST(FSerialization_BitSetLoadOutdated, "System.Mass.Serialization.BitSet.LoadOutdated");
+
+
+struct FSerialization_BitSetOverride : FExecutionTestBase
+{
+	virtual bool InstantTest() override
+	{
+		FTestStructTracker TestTracker;
+		const int32 FloatIndex = TestTracker.Add<FTestFragment_Float>();
+		const int32 IntIndex = TestTracker.Add<FTestFragment_Int>();
+		FStructTypeBitSet::FBitSetContainer TestBitArray;
+		TestBitArray.AddAtIndex(FloatIndex);
+		TestBitArray.AddAtIndex(IntIndex);
+
+		// create a bitset we're expected to get while loading TestBitArray's contents into FMassFragmentBitSet
+		FMassFragmentBitSet RegularBitSet;
+		RegularBitSet.Add<FTestFragment_Float>();;
+		RegularBitSet.Add<FTestFragment_Int>();
+
+		// save
+		TArray<uint8> Data;
+		FMemoryWriter Writer(Data);
+		TestTracker.Serialize(Writer, TestBitArray);
+
+		// load
+		FMemoryReader Reader(Data);
+		FMassFragmentBitSet SerializedBitSet;
+		// KEY STEP - add some crap to pollute the bitset
+		SerializedBitSet.Add<FTestFragment_Large>();
+		SerializedBitSet.Add<FTestFragment_Array>();
+
+		SerializedBitSet.Serialize(Reader);
+
+		AITEST_TRUE("The original contents of SerializedBitSet should have been erased", RegularBitSet.IsEquivalent(SerializedBitSet));
+
+		return true;
+	}
+};
+IMPLEMENT_AI_INSTANT_TEST(FSerialization_BitSetOverride, "System.Mass.Serialization.BitSet.OverrideExisting");
+
+} // FMassSerializationTest
+
+UE_ENABLE_OPTIMIZATION_SHIP
+
+#undef LOCTEXT_NAMESPACE
