Static analyzer warnings in StructTypeBitSet.cpp fix
Build issue #267618

diff --git a/Engine/Plugins/Experimental/StructUtils/Source/StructUtils/Private/StructTypeBitSet.cpp b/Engine/Plugins/Experimental/StructUtils/Source/StructUtils/Private/StructTypeBitSet.cpp
--- a/Engine/Plugins/Experimental/StructUtils/Source/StructUtils/Private/StructTypeBitSet.cpp
+++ b/Engine/Plugins/Experimental/StructUtils/Source/StructUtils/Private/StructTypeBitSet.cpp
@@ -15,7 +15,12 @@ namespace FBitSetSerializationDictionary
 	{
 		if (const FStructTracker** StoredTracer = TrackerMap.Find(SerializationHash))
 		{
-			ensureMsgf(*StoredTracer == nullptr || *StoredTracer == &Tracker
+			CA_SUPPRESS(6269); // warning C6269: Possibly incorrect order of operations.
+			// disabling the warning since it's wrong. 
+			// adding:
+			//		const FStructTracker* StoredTracerPtr = *StoredTracer;
+			// and then using StoredTracerPtr clears out the warning.
+			ensureMsgf(((*StoredTracer) == nullptr) || ((*StoredTracer) == &Tracker)
 				, TEXT("Hash conflict when registering a FStructTracker instance"));
 		}
 		TrackerMap.Add(SerializationHash, &Tracker);
