Moved all FInstancedStruct code that depends on the Engine module into a new StructUtilsEngine module.

This allows modules that can't depend on Engine to use FInstancedStruct.

diff --git a/Engine/Plugins/Experimental/StructUtils/Config/BaseStructUtils.ini b/Engine/Plugins/Experimental/StructUtils/Config/BaseStructUtils.ini
--- a/Engine/Plugins/Experimental/StructUtils/Config/BaseStructUtils.ini
+++ b/Engine/Plugins/Experimental/StructUtils/Config/BaseStructUtils.ini
@@ -1,3 +1,4 @@
 [CoreRedirects]
 +StructRedirects=(OldName="/Script/StructUtils.UniqueScriptStructPtr", NewName="/Script/StructUtils.InstancedStruct")
 +StructRedirects=(OldName="/Script/StructUtils.InstancedStructArray", NewName="/Script/StructUtils.InstancedStructContainer")
++ClassRedirects=(OldName="/Script/StructUtils.StructUtilsFunctionLibrary",NewName="/Script/StructUtilsEngine.StructUtilsFunctionLibrary")
diff --git a/Engine/Plugins/Experimental/StructUtils/Source/StructUtils/Private/InstancedStruct.cpp b/Engine/Plugins/Experimental/StructUtils/Source/StructUtils/Private/InstancedStruct.cpp
--- a/Engine/Plugins/Experimental/StructUtils/Source/StructUtils/Private/InstancedStruct.cpp
+++ b/Engine/Plugins/Experimental/StructUtils/Source/StructUtils/Private/InstancedStruct.cpp
@@ -2,9 +2,12 @@
 #include "InstancedStruct.h"
 #include "StructView.h"
 #include "Serialization/PropertyLocalizationDataGathering.h"
+
+#if WITH_ENGINE
 #include "Engine/PackageMapClient.h"
 #include "Engine/NetConnection.h"
 #include "Net/RepLayout.h"
+#endif // WITH_ENGINE
 
 #include UE_INLINE_GENERATED_CPP_BY_NAME(InstancedStruct)
 
@@ -389,6 +392,7 @@ void FInstancedStruct::AddStructReferencedObjects(class FReferenceCollector& Col
 
 bool FInstancedStruct::NetSerialize(FArchive& Ar, UPackageMap* Map, bool& bOutSuccess)
 {
+#if WITH_ENGINE
 	uint8 bValidData = Ar.IsSaving() ? IsValid() : 0;
 	Ar.SerializeBits(&bValidData, 1);
 
@@ -457,4 +461,11 @@ bool FInstancedStruct::NetSerialize(FArchive& Ar, UPackageMap* Map, bool& bOutSu
 	}
 
 	return true;
+
+#else // WITH_ENGINE
+
+	// The implementation above relies on types in the Engine module, so it can't be compiled without the engine.
+	return false;
+
+#endif // WITH_ENGINE
 }
diff --git a/Engine/Plugins/Experimental/StructUtils/Source/StructUtils/Private/StructUtilsFunctionLibrary.cpp b/Engine/Plugins/Experimental/StructUtils/Source/StructUtils/Private/StructUtilsFunctionLibrary.cpp
deleted file mode 100644
diff --git a/Engine/Plugins/Experimental/StructUtils/Source/StructUtils/Public/InstancedStruct.h b/Engine/Plugins/Experimental/StructUtils/Source/StructUtils/Public/InstancedStruct.h
--- a/Engine/Plugins/Experimental/StructUtils/Source/StructUtils/Public/InstancedStruct.h
+++ b/Engine/Plugins/Experimental/StructUtils/Source/StructUtils/Public/InstancedStruct.h
@@ -18,7 +18,7 @@ struct FConstSharedStruct;
  *	UPROPERTY(EditAnywhere, Category = Foo, meta = (BaseStruct = "/Script/ModuleName.TestStructBase"))
  *	TArray<FInstancedStruct> TestArray;
  */
-USTRUCT(BlueprintType, meta = (HasNativeMake = "/Script/StructUtils.StructUtilsFunctionLibrary.MakeInstancedStruct"))
+USTRUCT(BlueprintType, meta = (HasNativeMake = "/Script/StructUtilsEngine.StructUtilsFunctionLibrary.MakeInstancedStruct"))
 struct STRUCTUTILS_API FInstancedStruct
 {
 	GENERATED_BODY()
@@ -270,4 +270,4 @@ namespace UE::StructUtils::Private
 {
 void RegisterInstancedStructForLocalization();
 }
-#endif // WITH_EDITORONLY_DATA
\ No newline at end of file
+#endif // WITH_EDITORONLY_DATA
diff --git a/Engine/Plugins/Experimental/StructUtils/Source/StructUtils/Public/StructUtilsFunctionLibrary.h b/Engine/Plugins/Experimental/StructUtils/Source/StructUtils/Public/StructUtilsFunctionLibrary.h
deleted file mode 100644
diff --git a/Engine/Plugins/Experimental/StructUtils/Source/StructUtils/StructUtils.Build.cs b/Engine/Plugins/Experimental/StructUtils/Source/StructUtils/StructUtils.Build.cs
--- a/Engine/Plugins/Experimental/StructUtils/Source/StructUtils/StructUtils.Build.cs
+++ b/Engine/Plugins/Experimental/StructUtils/Source/StructUtils/StructUtils.Build.cs
@@ -12,9 +12,15 @@ namespace UnrealBuildTool.Rules
 			new string[] {
 				"Core",
 				"CoreUObject",
-				"Engine",
 			}
 			);
+
+			// Code such as FInstancedStruct::NetSerialize relies on the engine, but if the engine
+			// isn't available, this code will be compiled out using #if WITH_ENGINE.
+			if (Target.bCompileAgainstEngine)
+			{
+				PrivateDependencyModuleNames.Add("Engine");
+			}
 		}
 	}
 }
diff --git a/Engine/Plugins/Experimental/StructUtils/Source/StructUtilsEngine/Private/StructUtilsEngineModule.cpp b/Engine/Plugins/Experimental/StructUtils/Source/StructUtilsEngine/Private/StructUtilsEngineModule.cpp
new file mode 100644
--- /dev/null
+++ b/Engine/Plugins/Experimental/StructUtils/Source/StructUtilsEngine/Private/StructUtilsEngineModule.cpp
@@ -0,0 +1,5 @@
+// Copyright Epic Games, Inc. All Rights Reserved.
+
+#include "StructUtilsEngineModule.h"
+
+IMPLEMENT_MODULE(FDefaultModuleImpl, StructUtilsEngine);
diff --git a/Engine/Plugins/Experimental/StructUtils/Source/StructUtilsEngine/Private/StructUtilsFunctionLibrary.cpp b/Engine/Plugins/Experimental/StructUtils/Source/StructUtilsEngine/Private/StructUtilsFunctionLibrary.cpp
new file mode 100644
--- /dev/null
+++ b/Engine/Plugins/Experimental/StructUtils/Source/StructUtilsEngine/Private/StructUtilsFunctionLibrary.cpp
@@ -0,0 +1,140 @@
+// Copyright Epic Games, Inc. All Rights Reserved.
+
+#include "StructUtilsFunctionLibrary.h"
+#include "UObject/EnumProperty.h"
+
+#include UE_INLINE_GENERATED_CPP_BY_NAME(StructUtilsFunctionLibrary)
+
+#define LOCTEXT_NAMESPACE "UStructUtilsFunctionLibrary"
+
+
+FInstancedStruct UStructUtilsFunctionLibrary::MakeInstancedStruct(const int32& Value)
+{
+	// We should never hit this! stubs to avoid NoExport on the class.
+	checkNoEntry();
+	return {};
+}
+
+void UStructUtilsFunctionLibrary::SetInstancedStructValue(FInstancedStruct& InstancedStruct, const int32& Value)
+{
+	// We should never hit this! stubs to avoid NoExport on the class.
+	checkNoEntry();
+}
+
+void UStructUtilsFunctionLibrary::GetInstancedStructValue(EStructUtilsResult& ExecResult, const FInstancedStruct& InstancedStruct, int32& Value)
+{
+	// We should never hit this! stubs to avoid NoExport on the class.
+	checkNoEntry();
+}
+
+DEFINE_FUNCTION(UStructUtilsFunctionLibrary::execMakeInstancedStruct)
+{
+	// Read wildcard Value input.
+	Stack.MostRecentPropertyAddress = nullptr;
+	Stack.MostRecentPropertyContainer = nullptr;
+	Stack.StepCompiledIn<FStructProperty>(nullptr);
+	
+	const FStructProperty* ValueProp = CastField<FStructProperty>(Stack.MostRecentProperty);
+	const void* ValuePtr = Stack.MostRecentPropertyAddress;
+
+	P_FINISH;
+
+	if (!ValueProp || !ValuePtr)
+	{
+		FBlueprintExceptionInfo ExceptionInfo(
+			EBlueprintExceptionType::AbortExecution,
+			LOCTEXT("InstancedStruct_MakeInvalidValueWarning", "Failed to resolve the Value for Make Instanced Struct")
+		);
+
+		FBlueprintCoreDelegates::ThrowScriptException(P_THIS, Stack, ExceptionInfo);
+
+		P_NATIVE_BEGIN;
+		(*(FInstancedStruct*)RESULT_PARAM).Reset();
+		P_NATIVE_END;
+	}
+	else
+	{
+		P_NATIVE_BEGIN;
+		(*(FInstancedStruct*)RESULT_PARAM).InitializeAs(ValueProp->Struct, (const uint8*)ValuePtr);
+		P_NATIVE_END;
+	}
+}
+
+DEFINE_FUNCTION(UStructUtilsFunctionLibrary::execSetInstancedStructValue)
+{
+	P_GET_STRUCT_REF(FInstancedStruct, InstancedStruct);
+
+	// Read wildcard Value input.
+	Stack.MostRecentPropertyAddress = nullptr;
+	Stack.MostRecentPropertyContainer = nullptr;
+	Stack.StepCompiledIn<FStructProperty>(nullptr);
+	
+	const FStructProperty* ValueProp = CastField<FStructProperty>(Stack.MostRecentProperty);
+	const void* ValuePtr = Stack.MostRecentPropertyAddress;
+
+	P_FINISH;
+
+	if (!ValueProp || !ValuePtr)
+	{
+		FBlueprintExceptionInfo ExceptionInfo(
+			EBlueprintExceptionType::AbortExecution,
+			LOCTEXT("InstancedStruct_SetInvalidValueWarning", "Failed to resolve the Value for Set Instanced Struct Value")
+		);
+
+		FBlueprintCoreDelegates::ThrowScriptException(P_THIS, Stack, ExceptionInfo);
+
+		P_NATIVE_BEGIN;
+		InstancedStruct.Reset();
+		P_NATIVE_END;
+	}
+	else
+	{
+		P_NATIVE_BEGIN;
+		InstancedStruct.InitializeAs(ValueProp->Struct, (const uint8*)ValuePtr);
+		P_NATIVE_END;
+	}
+}
+
+DEFINE_FUNCTION(UStructUtilsFunctionLibrary::execGetInstancedStructValue)
+{
+	P_GET_ENUM_REF(EStructUtilsResult, ExecResult);
+	P_GET_STRUCT_REF(FInstancedStruct, InstancedStruct);
+
+	// Read wildcard Value input.
+	Stack.MostRecentPropertyAddress = nullptr;
+	Stack.MostRecentPropertyContainer = nullptr;
+	Stack.StepCompiledIn<FStructProperty>(nullptr);
+	
+	const FStructProperty* ValueProp = CastField<FStructProperty>(Stack.MostRecentProperty);
+	void* ValuePtr = Stack.MostRecentPropertyAddress;
+
+	P_FINISH;
+
+	ExecResult = EStructUtilsResult::NotValid;
+
+	if (!ValueProp || !ValuePtr)
+	{
+		FBlueprintExceptionInfo ExceptionInfo(
+			EBlueprintExceptionType::AbortExecution,
+			LOCTEXT("InstancedStruct_GetInvalidValueWarning", "Failed to resolve the Value for Get Instanced Struct Value")
+		);
+
+		FBlueprintCoreDelegates::ThrowScriptException(P_THIS, Stack, ExceptionInfo);
+	}
+	else
+	{
+		P_NATIVE_BEGIN;
+		if (InstancedStruct.IsValid() && InstancedStruct.GetScriptStruct()->IsChildOf(ValueProp->Struct))
+		{
+			ValueProp->Struct->CopyScriptStruct(ValuePtr, InstancedStruct.GetMemory());
+			ExecResult = EStructUtilsResult::Valid;
+		}
+		else
+		{
+			ExecResult = EStructUtilsResult::NotValid;
+		}
+		P_NATIVE_END;
+	}
+}
+
+#undef LOCTEXT_NAMESPACE
diff --git a/Engine/Plugins/Experimental/StructUtils/Source/StructUtilsEngine/Public/StructUtilsEngineModule.h b/Engine/Plugins/Experimental/StructUtils/Source/StructUtilsEngine/Public/StructUtilsEngineModule.h
new file mode 100644
--- /dev/null
+++ b/Engine/Plugins/Experimental/StructUtils/Source/StructUtilsEngine/Public/StructUtilsEngineModule.h
@@ -0,0 +1,5 @@
+// Copyright Epic Games, Inc. All Rights Reserved.
+
+#pragma once
+
+#include "Modules/ModuleManager.h"
diff --git a/Engine/Plugins/Experimental/StructUtils/Source/StructUtilsEngine/Public/StructUtilsFunctionLibrary.h b/Engine/Plugins/Experimental/StructUtils/Source/StructUtilsEngine/Public/StructUtilsFunctionLibrary.h
new file mode 100644
--- /dev/null
+++ b/Engine/Plugins/Experimental/StructUtils/Source/StructUtilsEngine/Public/StructUtilsFunctionLibrary.h
@@ -0,0 +1,83 @@
+// Copyright Epic Games, Inc. All Rights Reserved.
+
+#pragma once
+
+#include "Kismet/BlueprintFunctionLibrary.h"
+#include "InstancedStruct.h"
+#include "StructUtilsFunctionLibrary.generated.h"
+
+struct FGenericStruct;
+
+UENUM()
+enum class EStructUtilsResult : uint8
+{
+	Valid,
+	NotValid,
+};
+
+UCLASS()
+class STRUCTUTILSENGINE_API UStructUtilsFunctionLibrary : public UBlueprintFunctionLibrary
+{
+	GENERATED_BODY()
+
+public:
+
+	/**
+	 * Create a new InstancedStruct from the given source value.
+	 */
+	UFUNCTION(BlueprintCallable, CustomThunk, Category = "Utilities|Instanced Struct", meta = (CustomStructureParam = "Value", BlueprintInternalUseOnly="true", NativeMakeFunc))
+	static FInstancedStruct MakeInstancedStruct(const int32& Value);
+
+	/**
+	 * Sets the value of InstancedStruct from the given source value.
+	 */
+	UFUNCTION(BlueprintCallable, CustomThunk, Category = "Utilities|Instanced Struct", meta = (CustomStructureParam = "Value", BlueprintInternalUseOnly="true"))
+	static void SetInstancedStructValue(UPARAM(Ref) FInstancedStruct& InstancedStruct, const int32& Value);
+
+	/**
+	 * Retrieves data from an InstancedStruct if it matches the output type.
+	 */
+	UFUNCTION(BlueprintCallable, CustomThunk, Category = "Utilities|Instanced Struct", meta = (CustomStructureParam = "Value", ExpandEnumAsExecs = "ExecResult", BlueprintInternalUseOnly="true"))
+	static void GetInstancedStructValue(EStructUtilsResult& ExecResult, UPARAM(Ref) const FInstancedStruct& InstancedStruct, int32& Value);
+
+	/**
+	 * Resets an InstancedStruct.
+	 */
+	UFUNCTION(BlueprintCallable, Category = "Utilities|Instanced Struct", meta = (AdvancedDisplay = "1"))
+	static void Reset(UPARAM(Ref)FInstancedStruct& InstancedStruct, const UScriptStruct* StructType = nullptr)
+	{
+		InstancedStruct.InitializeAs(StructType, nullptr);
+	}
+
+	/** 
+	 * Checks whether an InstancedStruct contains value.
+	 */
+	UFUNCTION(BlueprintCallable, Category = "Utilities|Instanced Struct", meta = (DisplayName = "Is Valid (Branch)", ExpandEnumAsExecs = "ReturnValue"))
+	static EStructUtilsResult IsInstancedStructValid(UPARAM(Ref) const FInstancedStruct& InstancedStruct)
+	{
+		return InstancedStruct.IsValid() ? EStructUtilsResult::Valid : EStructUtilsResult::NotValid;
+	}
+
+	/**
+	 * Checks whether two InstancedStructs (and the values contained within) are equal.
+	 */
+	UFUNCTION(BlueprintPure, Category = "Utilities|Instanced Struct", meta = (CompactNodeTitle = "==", DisplayName = "Equal", Keywords = "== equal"))
+	static bool EqualEqual_InstancedStruct(UPARAM(Ref) const FInstancedStruct& A, UPARAM(Ref) const FInstancedStruct& B) { return A == B; }
+
+	/**
+	 * Checks whether two InstancedStructs are not equal.
+	 */
+	UFUNCTION(BlueprintPure, Category = "Utilities|Instanced Struct", meta = (CompactNodeTitle = "!=", DisplayName = "Not Equal", Keywords = "!= not equal"))
+	static bool NotEqual_InstancedStruct(UPARAM(Ref) const FInstancedStruct& A, UPARAM(Ref) const FInstancedStruct& B) { return A != B; }
+
+	/**
+	 * Checks whether the InstancedStruct contains value.
+	 */
+	UFUNCTION(BlueprintPure, Category = "Utilities|Instanced Struct", meta = ( DisplayName = "Is Valid"))
+	static bool IsValid_InstancedStruct(UPARAM(Ref) const FInstancedStruct& InstancedStruct) { return InstancedStruct.IsValid(); }
+
+private:
+	DECLARE_FUNCTION(execMakeInstancedStruct);
+	DECLARE_FUNCTION(execSetInstancedStructValue);
+	DECLARE_FUNCTION(execGetInstancedStructValue);
+};
\ No newline at end of file
diff --git a/Engine/Plugins/Experimental/StructUtils/Source/StructUtilsEngine/StructUtilsEngine.Build.cs b/Engine/Plugins/Experimental/StructUtils/Source/StructUtilsEngine/StructUtilsEngine.Build.cs
new file mode 100644
--- /dev/null
+++ b/Engine/Plugins/Experimental/StructUtils/Source/StructUtilsEngine/StructUtilsEngine.Build.cs
@@ -0,0 +1,21 @@
+// Copyright Epic Games, Inc. All Rights Reserved.
+
+namespace UnrealBuildTool.Rules
+{
+	public class StructUtilsEngine : ModuleRules
+	{
+		public StructUtilsEngine(ReadOnlyTargetRules Target) : base(Target)
+		{
+			PCHUsage = ModuleRules.PCHUsageMode.UseExplicitOrSharedPCHs;
+
+			PublicDependencyModuleNames.AddRange(
+			new string[] {
+				"Core",
+				"CoreUObject",
+				"Engine",
+				"StructUtils"
+			}
+			);
+		}
+	}
+}
diff --git a/Engine/Plugins/Experimental/StructUtils/Source/StructUtilsNodes/StructUtilsNodes.Build.cs b/Engine/Plugins/Experimental/StructUtils/Source/StructUtilsNodes/StructUtilsNodes.Build.cs
--- a/Engine/Plugins/Experimental/StructUtils/Source/StructUtilsNodes/StructUtilsNodes.Build.cs
+++ b/Engine/Plugins/Experimental/StructUtils/Source/StructUtilsNodes/StructUtilsNodes.Build.cs
@@ -18,6 +18,7 @@ namespace UnrealBuildTool.Rules
 					"KismetCompiler",
 					"PropertyEditor",
 					"StructUtils",
+					"StructUtilsEngine",
 					"UnrealEd"
 				}
 			);
diff --git a/Engine/Plugins/Experimental/StructUtils/StructUtils.uplugin b/Engine/Plugins/Experimental/StructUtils/StructUtils.uplugin
--- a/Engine/Plugins/Experimental/StructUtils/StructUtils.uplugin
+++ b/Engine/Plugins/Experimental/StructUtils/StructUtils.uplugin
@@ -21,6 +21,10 @@
 			"Type": "Runtime",
 			"LoadingPhase": "PreDefault"
 		},
+		{
+			"Name": "StructUtilsEngine",
+			"Type": "Runtime"
+		},
 		{
 			"Name" : "StructUtilsEditor",
 			"Type" : "Editor"
