Correctly report the nested preload dependencies of an instanced struct

#rnx

diff --git a/Engine/Plugins/Experimental/StructUtils/Source/StructUtils/Private/InstancedStruct.cpp b/Engine/Plugins/Experimental/StructUtils/Source/StructUtils/Private/InstancedStruct.cpp
--- a/Engine/Plugins/Experimental/StructUtils/Source/StructUtils/Private/InstancedStruct.cpp
+++ b/Engine/Plugins/Experimental/StructUtils/Source/StructUtils/Private/InstancedStruct.cpp
@@ -362,6 +362,24 @@ void FInstancedStruct::GetPreloadDependencies(TArray<UObject*>& OutDeps)
 	if (UScriptStruct* NonConstStruct = const_cast<UScriptStruct*>(GetScriptStruct()))
 	{
 		OutDeps.Add(NonConstStruct);
+
+		// Report direct dependencies of the instanced struct
+		if (UScriptStruct::ICppStructOps* CppStructOps = GetScriptStruct()->GetCppStructOps())
+		{
+			CppStructOps->GetPreloadDependencies(GetMutableMemory(), OutDeps);
+		}
+
+		// Report indirect dependencies of the instanced struct
+		// The iterator will recursively loop through all structs in structs/containers too
+		for (TPropertyValueIterator<FStructProperty> It(GetScriptStruct(), GetMutableMemory()); It; ++It)
+		{
+			const UScriptStruct* StructType = It.Key()->Struct;
+			if (UScriptStruct::ICppStructOps* CppStructOps = StructType->GetCppStructOps())
+			{
+				void* StructDataPtr = const_cast<void*>(It.Value());
+				CppStructOps->GetPreloadDependencies(StructDataPtr, OutDeps);
+			}
+		}
 	}
 }
 
diff --git a/Engine/Plugins/Experimental/StructUtils/Source/StructUtils/Private/InstancedStructContainer.cpp b/Engine/Plugins/Experimental/StructUtils/Source/StructUtils/Private/InstancedStructContainer.cpp
--- a/Engine/Plugins/Experimental/StructUtils/Source/StructUtils/Private/InstancedStructContainer.cpp
+++ b/Engine/Plugins/Experimental/StructUtils/Source/StructUtils/Private/InstancedStructContainer.cpp
@@ -576,6 +576,24 @@ void FInstancedStructContainer::GetPreloadDependencies(TArray<UObject*>& OutDeps
 		if (UScriptStruct* NonConstStruct = const_cast<UScriptStruct*>(Item.ScriptStruct))
 		{
 			OutDeps.Add(NonConstStruct);
+
+			// Report direct dependencies of the instanced struct
+			if (UScriptStruct::ICppStructOps* CppStructOps = Item.ScriptStruct->GetCppStructOps())
+			{
+				CppStructOps->GetPreloadDependencies(Memory + Item.Offset, OutDeps);
+			}
+
+			// Report indirect dependencies of the instanced struct
+			// The iterator will recursively loop through all structs in structs/containers too
+			for (TPropertyValueIterator<FStructProperty> It(Item.ScriptStruct, Memory + Item.Offset); It; ++It)
+			{
+				const UScriptStruct* StructType = It.Key()->Struct;
+				if (UScriptStruct::ICppStructOps* CppStructOps = StructType->GetCppStructOps())
+				{
+					void* StructDataPtr = const_cast<void*>(It.Value());
+					CppStructOps->GetPreloadDependencies(StructDataPtr, OutDeps);
+				}
+			}
 		}
 	}
 }
