PR #10248: Use property handle's metadata in FInstancedStructDetails (Contributed by [REDACTED_PERSON])

diff --git a/Engine/Plugins/Experimental/StructUtils/Source/StructUtilsEditor/Private/Customizations/InstancedStructDetails.cpp b/Engine/Plugins/Experimental/StructUtils/Source/StructUtilsEditor/Private/Customizations/InstancedStructDetails.cpp
--- a/Engine/Plugins/Experimental/StructUtils/Source/StructUtilsEditor/Private/Customizations/InstancedStructDetails.cpp
+++ b/Engine/Plugins/Experimental/StructUtils/Source/StructUtilsEditor/Private/Customizations/InstancedStructDetails.cpp
@@ -347,13 +347,11 @@ void FInstancedStructDetails::CustomizeHeader(TSharedRef<class IPropertyHandle>
 
 	StructProperty = StructPropertyHandle;
 
-	const FProperty* MetaDataProperty = StructProperty->GetMetaDataProperty();
-
-	const bool bEnableStructSelection = !MetaDataProperty->HasMetaData(NAME_StructTypeConst);
+	const bool bEnableStructSelection = !StructProperty->HasMetaData(NAME_StructTypeConst);
 
 	BaseScriptStruct = nullptr;
 	{
-		const FString& BaseStructName = MetaDataProperty->GetMetaData(NAME_BaseStruct);
+		const FString& BaseStructName = StructProperty->GetMetaData(NAME_BaseStruct);
 		if (!BaseStructName.IsEmpty())
 		{
 			BaseScriptStruct = UClass::TryFindTypeSlow<UScriptStruct>(BaseStructName);
@@ -427,12 +425,10 @@ TSharedRef<SWidget> FInstancedStructDetails::GenerateStructPicker()
 	static const FName NAME_HideViewOptions = "HideViewOptions";
 	static const FName NAME_ShowTreeView = "ShowTreeView";
 
-	const FProperty* MetaDataProperty = StructProperty->GetMetaDataProperty();
-
-	const bool bExcludeBaseStruct = MetaDataProperty->HasMetaData(NAME_ExcludeBaseStruct);
-	const bool bAllowNone = !(MetaDataProperty->PropertyFlags & CPF_NoClear);
-	const bool bHideViewOptions = MetaDataProperty->HasMetaData(NAME_HideViewOptions);
-	const bool bShowTreeView = MetaDataProperty->HasMetaData(NAME_ShowTreeView);
+	const bool bExcludeBaseStruct = StructProperty->HasMetaData(NAME_ExcludeBaseStruct);
+	const bool bAllowNone = !(StructProperty->GetMetaDataProperty()->PropertyFlags & CPF_NoClear);
+	const bool bHideViewOptions = StructProperty->HasMetaData(NAME_HideViewOptions);
+	const bool bShowTreeView = StructProperty->HasMetaData(NAME_ShowTreeView);
 
 	TSharedRef<FInstancedStructFilter> StructFilter = MakeShared<FInstancedStructFilter>();
 	StructFilter->BaseStruct = BaseScriptStruct;
