StructUtils: Do not try to compare mismatching structs during GatherForLocalization.

diff --git a/Engine/Plugins/Experimental/StructUtils/Source/StructUtils/Private/InstancedStruct.cpp b/Engine/Plugins/Experimental/StructUtils/Source/StructUtils/Private/InstancedStruct.cpp
--- a/Engine/Plugins/Experimental/StructUtils/Source/StructUtils/Private/InstancedStruct.cpp
+++ b/Engine/Plugins/Experimental/StructUtils/Source/StructUtils/Private/InstancedStruct.cpp
@@ -23,7 +23,17 @@ void GatherForLocalization(const FString& PathToParent, const UScriptStruct* Str
 
 	if (const UScriptStruct* StructTypePtr = ThisInstance->GetScriptStruct())
 	{
-		PropertyLocalizationDataGatherer.GatherLocalizationDataFromStructWithCallbacks(PathToParent + TEXT(".StructInstance"), StructTypePtr, ThisInstance->GetMemory(), DefaultInstance ? DefaultInstance->GetMemory() : nullptr, GatherTextFlags);
+		const uint8* DefaultInstanceMemory = nullptr;
+		if (DefaultInstance)
+		{
+			// Types must match
+			if (StructTypePtr == DefaultInstance->GetScriptStruct())
+			{
+				DefaultInstanceMemory = DefaultInstance->GetMemory();
+			}
+		}
+		
+		PropertyLocalizationDataGatherer.GatherLocalizationDataFromStructWithCallbacks(PathToParent + TEXT(".StructInstance"), StructTypePtr, ThisInstance->GetMemory(), DefaultInstanceMemory, GatherTextFlags);
 	}
 }
 
