StructUtils: Made a few modifications to property bags so they can be used by [REDACTED_PERSON].

* Added Get/SetValueSerializedString(). MRQ needs to be able to set struct-valued properties via Python, and doing this via string is one of the only options feasible right now.
* Added "Count" to the EPropertyBagPropertyType and EPropertyBagContainerType enums. MRQ has similar enums, and Count helps ensure that MRQ's enums do not diverge from what property bags have.
* Updated the details customization to respect the "ShowOnlyInnerProperties" metadata.

diff --git a/Engine/Plugins/Experimental/StructUtils/Source/StructUtils/Private/PropertyBag.cpp b/Engine/Plugins/Experimental/StructUtils/Source/StructUtils/Private/PropertyBag.cpp
--- a/Engine/Plugins/Experimental/StructUtils/Source/StructUtils/Private/PropertyBag.cpp
+++ b/Engine/Plugins/Experimental/StructUtils/Source/StructUtils/Private/PropertyBag.cpp
@@ -1441,6 +1441,27 @@ TValueOrError<UClass*, EPropertyBagResult> FInstancedPropertyBag::GetValueClass(
 	return MakeValue(Class);
 }
 
+TValueOrError<FString, EPropertyBagResult> FInstancedPropertyBag::GetValueSerializedString(const FName Name)
+{
+	const FPropertyBagPropertyDesc* Desc = FindPropertyDescByName(Name);
+	if (Desc == nullptr)
+	{
+		return MakeError(EPropertyBagResult::PropertyNotFound);
+	}
+	
+	const FProperty* Property = Desc->CachedProperty;
+	check(Property);
+
+	const void* ValueAddress = GetValueAddress(Desc);
+	FString OutStringValue;
+	if (!Property->ExportText_Direct(OutStringValue, ValueAddress, ValueAddress, nullptr, PPF_None))
+	{
+		UE_LOG(LogCore, Warning, TEXT("PropertyBag: Getting the serialized value of the property '%s' failed."), *Desc->Name.ToString());
+		return MakeError(EPropertyBagResult::TypeMismatch);
+	}
+
+	return MakeValue(OutStringValue);
+}
 
 EPropertyBagResult FInstancedPropertyBag::SetValueBool(const FName Name, const bool bInValue)
 {
@@ -1520,6 +1541,26 @@ EPropertyBagResult FInstancedPropertyBag::SetValueClass(const FName Name, UClass
 	return UE::StructUtils::Private::SetPropertyValueAsObject(Desc, GetMutableValueAddress(Desc), InValue);
 }
 
+EPropertyBagResult FInstancedPropertyBag::SetValueSerializedString(const FName Name, const FString& InValue)
+{
+	const FPropertyBagPropertyDesc* Desc = FindPropertyDescByName(Name);
+	if (Desc == nullptr)
+	{
+		return EPropertyBagResult::PropertyNotFound;
+	}
+	
+	const FProperty* Property = Desc->CachedProperty;
+	check(Property);
+
+	if (!Property->ImportText_Direct(*InValue, GetMutableValueAddress(Desc), nullptr, PPF_None))
+	{
+		UE_LOG(LogCore, Warning, TEXT("PropertyBag: Setting the value of the property '%s' failed because the string representation provided was not accepted."), *Desc->Name.ToString());
+		return EPropertyBagResult::TypeMismatch;
+	}
+
+	return EPropertyBagResult::Success;
+}
+
 EPropertyBagResult FInstancedPropertyBag::SetValue(const FName Name, const FProperty* InSourceProperty, const void* InSourceContainerAddress)
 {
 	const FPropertyBagPropertyDesc* Desc = FindPropertyDescByName(Name);
diff --git a/Engine/Plugins/Experimental/StructUtils/Source/StructUtils/Public/PropertyBag.h b/Engine/Plugins/Experimental/StructUtils/Source/StructUtils/Public/PropertyBag.h
--- a/Engine/Plugins/Experimental/StructUtils/Source/StructUtils/Public/PropertyBag.h
+++ b/Engine/Plugins/Experimental/StructUtils/Source/StructUtils/Public/PropertyBag.h
@@ -27,6 +27,8 @@ enum class EPropertyBagPropertyType : uint8
 	SoftObject UMETA(Hidden),
 	Class UMETA(Hidden),
 	SoftClass UMETA(Hidden),
+
+	Count UMETA(Hidden)
 };
 
 /** Property bag property container type. */
@@ -35,6 +37,8 @@ enum class EPropertyBagContainerType : uint8
 {
 	None UMETA(Hidden),
 	Array,
+
+	Count UMETA(Hidden)
 };
 
 /** Getter and setter result code. */
@@ -332,6 +336,9 @@ struct STRUCTUTILS_API FInstancedPropertyBag
 	TValueOrError<UObject*, EPropertyBagResult> GetValueObject(const FName Name, const UClass* RequestedClass = nullptr) const;
 	TValueOrError<UClass*, EPropertyBagResult> GetValueClass(const FName Name) const;
 
+	/** @return string-based serialized representation of the value. */
+	TValueOrError<FString, EPropertyBagResult> GetValueSerializedString(const FName Name);
+
 	/** @return enum value of specified type. */
 	template <typename T>
 	TValueOrError<T, EPropertyBagResult> GetValueEnum(const FName Name) const
@@ -402,6 +409,12 @@ struct STRUCTUTILS_API FInstancedPropertyBag
 	EPropertyBagResult SetValueObject(const FName Name, UObject* InValue);
 	EPropertyBagResult SetValueClass(const FName Name, UClass* InValue);
 
+	/**
+	 * Sets property value from a serialized representation of the value. If the string value provided
+	 * cannot be parsed by the property, the operation will fail.
+	 */
+	EPropertyBagResult SetValueSerializedString(const FName Name, const FString& InValue);
+
 	/** Sets enum value specified type. */
 	template <typename T>
 	EPropertyBagResult SetValueEnum(const FName Name, const T InValue)
diff --git a/Engine/Plugins/Experimental/StructUtils/Source/StructUtilsEditor/Private/Customizations/PropertyBagDetails.cpp b/Engine/Plugins/Experimental/StructUtils/Source/StructUtilsEditor/Private/Customizations/PropertyBagDetails.cpp
--- a/Engine/Plugins/Experimental/StructUtils/Source/StructUtilsEditor/Private/Customizations/PropertyBagDetails.cpp
+++ b/Engine/Plugins/Experimental/StructUtils/Source/StructUtilsEditor/Private/Customizations/PropertyBagDetails.cpp
@@ -755,9 +755,17 @@ void FPropertyBagDetails::CustomizeHeader(TSharedRef<IPropertyHandle> StructProp
 	check(StructProperty);
 	
 	static const FName NAME_FixedLayout = "FixedLayout";
+	static const FName NAME_ShowOnlyInnerProperties = "ShowOnlyInnerProperties";
+	
 	if (const FProperty* MetaDataProperty = StructProperty->GetMetaDataProperty())
 	{
 		bFixedLayout = MetaDataProperty->HasMetaData(NAME_FixedLayout);
+
+		// Don't show the header if ShowOnlyInnerProperties is set
+		if (MetaDataProperty->HasMetaData(NAME_ShowOnlyInnerProperties))
+		{
+			return;
+		}
 	}
 
 	TSharedPtr<SWidget> ValueWidget = SNullWidget::NullWidget;
