StructUtils: Refresh instanced struct details after BP compile to prevent access to stale data.

diff --git a/Engine/Plugins/Experimental/StructUtils/Source/StructUtilsEditor/Private/Customizations/InstancedStructDetails.cpp b/Engine/Plugins/Experimental/StructUtils/Source/StructUtilsEditor/Private/Customizations/InstancedStructDetails.cpp
--- a/Engine/Plugins/Experimental/StructUtils/Source/StructUtilsEditor/Private/Customizations/InstancedStructDetails.cpp
+++ b/Engine/Plugins/Experimental/StructUtils/Source/StructUtilsEditor/Private/Customizations/InstancedStructDetails.cpp
@@ -14,6 +14,7 @@
 #include "Styling/SlateIconFinder.h"
 #include "Engine/UserDefinedStruct.h"
 #include "InstancedStruct.h"
+#include "IPropertyUtilities.h"
 #include "Widgets/Layout/SBox.h"
 
 #define LOCTEXT_NAMESPACE "StructUtilsEditor"
@@ -340,13 +341,24 @@ TSharedRef<IPropertyTypeCustomization> FInstancedStructDetails::MakeInstance()
 	return MakeShared<FInstancedStructDetails>();
 }
 
+FInstancedStructDetails::~FInstancedStructDetails()
+{
+	if (OnObjectsReinstancedHandle.IsValid())
+	{
+		FCoreUObjectDelegates::OnObjectsReinstanced.Remove(OnObjectsReinstancedHandle);
+	}
+}
+
 void FInstancedStructDetails::CustomizeHeader(TSharedRef<class IPropertyHandle> StructPropertyHandle, class FDetailWidgetRow& HeaderRow, IPropertyTypeCustomizationUtils& StructCustomizationUtils)
 {
 	static const FName NAME_BaseStruct = "BaseStruct";
 	static const FName NAME_StructTypeConst = "StructTypeConst";
 
 	StructProperty = StructPropertyHandle;
+	PropUtils = StructCustomizationUtils.GetPropertyUtilities();
 
+	OnObjectsReinstancedHandle = FCoreUObjectDelegates::OnObjectsReinstanced.AddSP(this, &FInstancedStructDetails::OnObjectsReinstanced);
+	
 	const bool bEnableStructSelection = !StructProperty->HasMetaData(NAME_StructTypeConst);
 
 	BaseScriptStruct = nullptr;
@@ -398,6 +410,15 @@ void FInstancedStructDetails::CustomizeHeader(TSharedRef<class IPropertyHandle>
 		];
 }
 
+void FInstancedStructDetails::OnObjectsReinstanced(const FReplacementObjectMap& ObjectMap)
+{
+	// Force update the details when BP is compiled, since we may cached hold references to the old object or class.
+	if (PropUtils.IsValid())
+	{
+		PropUtils->ForceRefresh();
+	}
+}
+
 void FInstancedStructDetails::CustomizeChildren(TSharedRef<class IPropertyHandle> StructPropertyHandle, class IDetailChildrenBuilder& StructBuilder, IPropertyTypeCustomizationUtils& StructCustomizationUtils)
 {
 	TSharedRef<FInstancedStructDataDetails> DataDetails = MakeShared<FInstancedStructDataDetails>(StructProperty);
diff --git a/Engine/Plugins/Experimental/StructUtils/Source/StructUtilsEditor/Public/InstancedStructDetails.h b/Engine/Plugins/Experimental/StructUtils/Source/StructUtilsEditor/Public/InstancedStructDetails.h
--- a/Engine/Plugins/Experimental/StructUtils/Source/StructUtilsEditor/Public/InstancedStructDetails.h
+++ b/Engine/Plugins/Experimental/StructUtils/Source/StructUtilsEditor/Public/InstancedStructDetails.h
@@ -18,6 +18,8 @@ struct FInstancedStruct;
 class STRUCTUTILSEDITOR_API FInstancedStructDetails : public IPropertyTypeCustomization
 {
 public:
+	virtual ~FInstancedStructDetails() override;
+	
 	/** Makes a new instance of this detail layout class for a specific detail view requesting it */
 	static TSharedRef<IPropertyTypeCustomization> MakeInstance();
 
@@ -26,6 +28,10 @@ public:
 	virtual void CustomizeChildren(TSharedRef<IPropertyHandle> StructPropertyHandle, class IDetailChildrenBuilder& StructBuilder, IPropertyTypeCustomizationUtils& StructCustomizationUtils) override;
 
 private:
+
+	using FReplacementObjectMap = TMap<UObject*, UObject*>;
+	void OnObjectsReinstanced(const FReplacementObjectMap& ObjectMap);
+	
 	FText GetDisplayValueString() const;
 	const FSlateBrush* GetDisplayValueIcon() const;
 	TSharedRef<SWidget> GenerateStructPicker();
@@ -33,11 +39,15 @@ private:
 
 	/** Handle to the struct property being edited */
 	TSharedPtr<IPropertyHandle> StructProperty;
-	
+
 	/** The base struct that we're allowing to be picked (controlled by the "BaseStruct" meta-data) */
 	UScriptStruct* BaseScriptStruct = nullptr;
 
 	TSharedPtr<class SComboButton> ComboButton;
+
+	TSharedPtr<IPropertyUtilities> PropUtils;
+	
+	FDelegateHandle OnObjectsReinstancedHandle;
 };
 
 /** 
