StructUtils: Be more prudent when to update properties on BP compiles.
- Fixedbug where opening a BP which contains instanced struct would take very-long-time

diff --git a/Engine/Plugins/Experimental/StructUtils/Source/StructUtilsEditor/Private/Customizations/InstancedStructDetails.cpp b/Engine/Plugins/Experimental/StructUtils/Source/StructUtilsEditor/Private/Customizations/InstancedStructDetails.cpp
--- a/Engine/Plugins/Experimental/StructUtils/Source/StructUtilsEditor/Private/Customizations/InstancedStructDetails.cpp
+++ b/Engine/Plugins/Experimental/StructUtils/Source/StructUtilsEditor/Private/Customizations/InstancedStructDetails.cpp
@@ -413,9 +413,9 @@ void FInstancedStructDetails::CustomizeHeader(TSharedRef<class IPropertyHandle>
 void FInstancedStructDetails::OnObjectsReinstanced(const FReplacementObjectMap& ObjectMap)
 {
 	// Force update the details when BP is compiled, since we may cached hold references to the old object or class.
-	if (PropUtils.IsValid())
+	if (!ObjectMap.IsEmpty() && PropUtils.IsValid())
 	{
-		PropUtils->ForceRefresh();
+		PropUtils->RequestRefresh();
 	}
 }
 
