StructUtils: Force property node update when struct changes.
- Fixes a bug where when changing structs really quickly, it would use old child nodes, and crash

diff --git a/Engine/Plugins/Experimental/StructUtils/Source/StructUtilsEditor/Private/Customizations/InstancedStructDetails.cpp b/Engine/Plugins/Experimental/StructUtils/Source/StructUtilsEditor/Private/Customizations/InstancedStructDetails.cpp
--- a/Engine/Plugins/Experimental/StructUtils/Source/StructUtilsEditor/Private/Customizations/InstancedStructDetails.cpp
+++ b/Engine/Plugins/Experimental/StructUtils/Source/StructUtilsEditor/Private/Customizations/InstancedStructDetails.cpp
@@ -512,6 +512,12 @@ void FInstancedStructDetails::OnStructPicked(const UScriptStruct* InStruct)
 
 		StructProperty->NotifyPostChange(EPropertyChangeType::ValueSet);
 		StructProperty->NotifyFinishedChangingProperties();
+
+		// Property tree will be invalid after changing the struct type, force update.
+		if (PropUtils.IsValid())
+		{
+			PropUtils->ForceRefresh();
+		}
 	}
 
 	ComboButton->SetIsOpen(false);
