Show ToolTip text for instance structs instead of DisplayName text.

diff --git a/Engine/Plugins/Experimental/StructUtils/Source/StructUtilsEditor/Private/Customizations/InstancedStructDetails.cpp b/Engine/Plugins/Experimental/StructUtils/Source/StructUtilsEditor/Private/Customizations/InstancedStructDetails.cpp
--- a/Engine/Plugins/Experimental/StructUtils/Source/StructUtilsEditor/Private/Customizations/InstancedStructDetails.cpp
+++ b/Engine/Plugins/Experimental/StructUtils/Source/StructUtilsEditor/Private/Customizations/InstancedStructDetails.cpp
@@ -401,7 +401,7 @@ void FInstancedStructDetails::CustomizeHeader(TSharedRef<class IPropertyHandle>
 				[
 					SNew(STextBlock)
 					.Text(this, &FInstancedStructDetails::GetDisplayValueString)
-					.ToolTipText(this, &FInstancedStructDetails::GetDisplayValueString)
+					.ToolTipText(this, &FInstancedStructDetails::GetTooltipText)
 					.Font(IDetailLayoutBuilder::GetDetailFont())
 				]
 			]
@@ -444,6 +444,16 @@ FText FInstancedStructDetails::GetDisplayValueString() const
 	return FText::GetEmpty();
 }
 
+FText FInstancedStructDetails::GetTooltipText() const
+{
+	const UScriptStruct* ScriptStruct = UE::StructUtils::Private::GetCommonScriptStruct(StructProperty);
+	if (ScriptStruct)
+	{
+		return ScriptStruct->GetToolTipText();
+	}
+	return GetDisplayValueString();
+}
+
 const FSlateBrush* FInstancedStructDetails::GetDisplayValueIcon() const
 {
 	const UScriptStruct* CommonStruct = nullptr;
diff --git a/Engine/Plugins/Experimental/StructUtils/Source/StructUtilsEditor/Public/InstancedStructDetails.h b/Engine/Plugins/Experimental/StructUtils/Source/StructUtilsEditor/Public/InstancedStructDetails.h
--- a/Engine/Plugins/Experimental/StructUtils/Source/StructUtilsEditor/Public/InstancedStructDetails.h
+++ b/Engine/Plugins/Experimental/StructUtils/Source/StructUtilsEditor/Public/InstancedStructDetails.h
@@ -34,6 +34,7 @@ private:
 	void OnObjectsReinstanced(const FReplacementObjectMap& ObjectMap);
 	
 	FText GetDisplayValueString() const;
+	FText GetTooltipText() const;
 	const FSlateBrush* GetDisplayValueIcon() const;
 	TSharedRef<SWidget> GenerateStructPicker();
 	void OnStructPicked(const UScriptStruct* InStruct);
