Added nested containers support to PropertyBag (TArray<TArray<T>>).
Currently supports nested arrays with a max depth of two.
PropertyBag UI does not currently support edition of nested containers and will disable the row and show a warning message (this requires Core engine support).
Part of UE-180115

diff --git a/Engine/Plugins/Experimental/StructUtils/Source/StructUtils/Private/PropertyBag.cpp b/Engine/Plugins/Experimental/StructUtils/Source/StructUtils/Private/PropertyBag.cpp
--- a/Engine/Plugins/Experimental/StructUtils/Source/StructUtils/Private/PropertyBag.cpp
+++ b/Engine/Plugins/Experimental/StructUtils/Source/StructUtils/Private/PropertyBag.cpp
@@ -16,6 +16,7 @@ struct STRUCTUTILS_API FPropertyBagCustomVersion
 
 		// Added support for array types
 		ContainerTypes = 1,
+		NestedContainerTypes = 2,
 
 		// -----<new versions can be added above this line>-------------------------------------------------
 		VersionPlusOne,
@@ -53,9 +54,9 @@ namespace UE::StructUtils::Private
 	uint64 CalcPropertyDescHash(const FPropertyBagPropertyDesc& Desc)
 	{
 #if WITH_EDITORONLY_DATA
-		const uint32 Hashes[] = { GetTypeHash(Desc.ID), GetTypeHash(Desc.Name), GetTypeHash(Desc.ValueType), GetTypeHash(Desc.ContainerType), GetTypeHash(Desc.MetaData) };
+		const uint32 Hashes[] = { GetTypeHash(Desc.ID), GetTypeHash(Desc.Name), GetTypeHash(Desc.ValueType), GetTypeHash(Desc.ContainerTypes), GetTypeHash(Desc.MetaData) };
 #else
-		const uint32 Hashes[] = { GetTypeHash(Desc.ID), GetTypeHash(Desc.Name), GetTypeHash(Desc.ValueType), GetTypeHash(Desc.ContainerType)};
+		const uint32 Hashes[] = { GetTypeHash(Desc.ID), GetTypeHash(Desc.Name), GetTypeHash(Desc.ValueType), GetTypeHash(Desc.ContainerTypes)};
 #endif
 		return CityHash64WithSeed((const char*)Hashes, sizeof(Hashes), GetObjectHash(Desc.ValueTypeObject));
 	}
@@ -70,14 +71,24 @@ namespace UE::StructUtils::Private
 		return Hash;
 	}
 
-	EPropertyBagContainerType GetContainerTypeFromProperty(const FProperty* InSourceProperty)
+	FPropertyBagContainerTypes GetContainerTypesFromProperty(const FProperty* InSourceProperty)
 	{
-		if (CastField<FArrayProperty>(InSourceProperty))
+		FPropertyBagContainerTypes ContainerTypes;
+
+		while (const FArrayProperty* ArrayProperty = CastField<FArrayProperty>(InSourceProperty))
 		{
-			return EPropertyBagContainerType::Array;
+			if (ContainerTypes.Add(EPropertyBagContainerType::Array))
+			{
+				InSourceProperty = CastField<FArrayProperty>(ArrayProperty->Inner);
+			}
+			else // we reached the nested containers limit
+			{
+				ContainerTypes.Reset();
+				break;
+			}
 		}
-		
-		return EPropertyBagContainerType::None;
+
+		return ContainerTypes;
 	}
 
 	EPropertyBagPropertyType GetValueTypeFromProperty(const FProperty* InSourceProperty)
@@ -198,15 +209,43 @@ namespace UE::StructUtils::Private
 
 	FProperty* CreatePropertyFromDesc(const FPropertyBagPropertyDesc& Desc, const FFieldVariant PropertyScope)
 	{
-		// Handle array property
-		if (Desc.ContainerType == EPropertyBagContainerType::Array)
+		// Handle array and nested containers properties
+		if (Desc.ContainerTypes.Num() > 0)
 		{
-			FArrayProperty* Prop = new FArrayProperty(PropertyScope, Desc.Name, RF_Public);
+			FProperty* Prop = nullptr; // the first created container will fill the return value, nested ones will fill the inner
+
+			// support for nested containers, i.e. : TArray<TArray<float>>
+			FFieldVariant PropertyOwner = PropertyScope;
+			FProperty** ValuePropertyPtr = &Prop;
+			FString ParentName;
 
+			// Create the container list
+			for (EPropertyBagContainerType BagContainerType : Desc.ContainerTypes)
+			{
+				switch(BagContainerType)
+				{
+				case EPropertyBagContainerType::Array:
+					{
+						// create an array property as a container for the tail
+						const FString PropName = ParentName.IsEmpty() ? Desc.Name.ToString() : ParentName + TEXT("_InnerArray");
+						FArrayProperty* ArrayProperty = new FArrayProperty(PropertyOwner, FName(PropName), RF_Public);
+						*ValuePropertyPtr = ArrayProperty;
+						ValuePropertyPtr = &ArrayProperty->Inner;
+						PropertyOwner = ArrayProperty;
+						ParentName = PropName;
+						break;
+					}
+				default:
+					ensureMsgf(false, TEXT("Unsuported container type %s"), *UEnum::GetValueAsString(BagContainerType));
+					break;
+				}
+			}
+
+			// finally create the tail type
 			FPropertyBagPropertyDesc InnerDesc = Desc;
-			InnerDesc.Name = FName(InnerDesc.Name.ToString() + TEXT("_inner"));
-			InnerDesc.ContainerType = EPropertyBagContainerType::None;
-			Prop->Inner = CreatePropertyFromDesc(InnerDesc, Prop);
+			InnerDesc.Name = FName(ParentName + TEXT("_InnerType"));
+			InnerDesc.ContainerTypes.Reset();
+			*ValuePropertyPtr = CreatePropertyFromDesc(InnerDesc, PropertyOwner);
 				
 			return Prop;
 		}
@@ -216,6 +255,7 @@ namespace UE::StructUtils::Private
 		case EPropertyBagPropertyType::Bool:
 			{
 				FBoolProperty* Prop = new FBoolProperty(PropertyScope, Desc.Name, RF_Public);
+				Prop->SetBoolSize(sizeof(bool), true); // Enable native access (init the whole byte, rather than just first bit)
 				return Prop;
 			}
 		case EPropertyBagPropertyType::Byte:
@@ -363,7 +403,7 @@ namespace UE::StructUtils::Private
 		{
 			return EPropertyBagResult::OutOfBounds;
 		}
-		if (Desc->ContainerType != EPropertyBagContainerType::None)
+		if (Desc->ContainerTypes.Num() > 0)
 		{
 			return EPropertyBagResult::TypeMismatch;
 		}
@@ -429,7 +469,7 @@ namespace UE::StructUtils::Private
 		{
 			return EPropertyBagResult::OutOfBounds;
 		}
-		if (Desc->ContainerType != EPropertyBagContainerType::None)
+		if (Desc->ContainerTypes.Num() > 0)
 		{
 			return EPropertyBagResult::TypeMismatch;
 		}
@@ -497,7 +537,7 @@ namespace UE::StructUtils::Private
 		{
 			return EPropertyBagResult::OutOfBounds;
 		}
-		if (Desc->ContainerType != EPropertyBagContainerType::None)
+		if (Desc->ContainerTypes.Num() > 0)
 		{
 			return EPropertyBagResult::TypeMismatch;
 		}
@@ -527,7 +567,7 @@ namespace UE::StructUtils::Private
 		{
 			return EPropertyBagResult::OutOfBounds;
 		}
-		if (Desc->ContainerType != EPropertyBagContainerType::None)
+		if (Desc->ContainerTypes.Num() > 0)
 		{
 			return EPropertyBagResult::TypeMismatch;
 		}
@@ -560,7 +600,7 @@ namespace UE::StructUtils::Private
 		{
 			return EPropertyBagResult::OutOfBounds;
 		}
-		if (Desc->ContainerType != EPropertyBagContainerType::None)
+		if (Desc->ContainerTypes.Num() > 0)
 		{
 			return EPropertyBagResult::TypeMismatch;
 		}
@@ -595,7 +635,7 @@ namespace UE::StructUtils::Private
 		{
 			return EPropertyBagResult::OutOfBounds;
 		}
-		if (Desc->ContainerType != EPropertyBagContainerType::None)
+		if (Desc->ContainerTypes.Num() > 0)
 		{
 			return EPropertyBagResult::TypeMismatch;
 		}
@@ -627,7 +667,7 @@ namespace UE::StructUtils::Private
 		{
 			return EPropertyBagResult::OutOfBounds;
 		}
-		if (Desc->ContainerType != EPropertyBagContainerType::None)
+		if (Desc->ContainerTypes.Num() > 0)
 		{
 			return EPropertyBagResult::TypeMismatch;
 		}
@@ -693,7 +733,7 @@ namespace UE::StructUtils::Private
 		{
 			return EPropertyBagResult::OutOfBounds;
 		}
-		if (Desc->ContainerType != EPropertyBagContainerType::None)
+		if (Desc->ContainerTypes.Num() > 0)
 		{
 			return EPropertyBagResult::TypeMismatch;
 		}
@@ -761,7 +801,7 @@ namespace UE::StructUtils::Private
 		{
 			return EPropertyBagResult::OutOfBounds;
 		}
-		if (Desc->ContainerType != EPropertyBagContainerType::None)
+		if (Desc->ContainerTypes.Num() > 0)
 		{
 			return EPropertyBagResult::TypeMismatch;
 		}
@@ -791,7 +831,7 @@ namespace UE::StructUtils::Private
 		{
 			return EPropertyBagResult::OutOfBounds;
 		}
-		if (Desc->ContainerType != EPropertyBagContainerType::None)
+		if (Desc->ContainerTypes.Num() > 0)
 		{
 			return EPropertyBagResult::TypeMismatch;
 		}
@@ -824,7 +864,7 @@ namespace UE::StructUtils::Private
 		{
 			return EPropertyBagResult::OutOfBounds;
 		}
-		if (Desc->ContainerType != EPropertyBagContainerType::None)
+		if (Desc->ContainerTypes.Num() > 0)
 		{
 			return EPropertyBagResult::TypeMismatch;
 		}
@@ -865,7 +905,7 @@ namespace UE::StructUtils::Private
 		{
 			return EPropertyBagResult::OutOfBounds;
 		}
-		if (Desc->ContainerType != EPropertyBagContainerType::None)
+		if (Desc->ContainerTypes.Num() > 0)
 		{
 			return EPropertyBagResult::TypeMismatch;
 		}
@@ -941,8 +981,8 @@ namespace UE::StructUtils::Private
 			{
 				TargetDesc.CachedProperty->CopyCompleteValue(TargetAddress, SourceAddress);
 			}
-			else if (TargetDesc.ContainerType == EPropertyBagContainerType::None
-					&& SourceDesc.ContainerType == EPropertyBagContainerType::None)
+			else if (TargetDesc.ContainerTypes.Num() == 0
+					&& SourceDesc.ContainerTypes.Num() == 0)
 			{
 				if (TargetDesc.IsNumericType() && SourceDesc.IsNumericType())
 				{
@@ -997,6 +1037,55 @@ namespace UE::StructUtils::Private
 }; // UE::StructUtils::Private
 
 
+//----------------------------------------------------------------//
+//  FPropertyBagContainerTypes
+//----------------------------------------------------------------//
+EPropertyBagContainerType FPropertyBagContainerTypes::PopHead()
+{
+	EPropertyBagContainerType Head = EPropertyBagContainerType::None;
+		
+	if (NumContainers > 0)
+	{
+		Head = Types[0];
+
+		for (uint32 i = 1; i < NumContainers; ++i)
+		{
+			Types[i - 1] = Types[i];
+		}
+		Types[NumContainers - 1] = EPropertyBagContainerType::None;
+		NumContainers = NumContainers - 1;
+	}
+
+	return Head;
+}
+
+void FPropertyBagContainerTypes::Serialize(FArchive& Ar)
+{
+	Ar << NumContainers;
+	for (int i = 0; i < NumContainers; ++i)
+	{
+		Ar << Types[i];
+	}
+}
+
+bool FPropertyBagContainerTypes::operator == (const FPropertyBagContainerTypes& Other) const
+{
+	if (NumContainers != Other.NumContainers)
+	{
+		return false;
+	}
+
+	for (int i = 0; i < NumContainers; ++i)
+	{
+		if (Types[i] != Other.Types[i])
+		{
+			return false;
+		}
+	}
+
+	return true;
+}
+
 //----------------------------------------------------------------//
 //  FPropertyBagPropertyDesc
 //----------------------------------------------------------------//
@@ -1012,7 +1101,8 @@ FPropertyBagPropertyDesc::FPropertyBagPropertyDesc(const FName InName, const FPr
 {
 	ValueType = UE::StructUtils::Private::GetValueTypeFromProperty(InSourceProperty);
 	ValueTypeObject = UE::StructUtils::Private::GetValueTypeObjectFromProperty(InSourceProperty);
-	ContainerType = UE::StructUtils::Private::GetContainerTypeFromProperty(InSourceProperty);
+	// @todo : improve error handling - if we reach the nested containers limit, the Desc will be invalid (empty container types)
+	ContainerTypes = UE::StructUtils::Private::GetContainerTypesFromProperty(InSourceProperty);
 
 #if WITH_EDITORONLY_DATA
 	if (const TMap<FName, FString>* SourcePropertyMetaData = InSourceProperty->GetMetaDataMap())
@@ -1034,7 +1124,17 @@ static FArchive& operator<<(FArchive& Ar, FPropertyBagPropertyDesc& Bag)
 
 	if (Ar.CustomVer(FPropertyBagCustomVersion::GUID) >= FPropertyBagCustomVersion::ContainerTypes)
 	{
-		Ar << Bag.ContainerType;
+		if (Ar.IsLoading() && Ar.CustomVer(FPropertyBagCustomVersion::GUID) < FPropertyBagCustomVersion::NestedContainerTypes)
+		{
+			EPropertyBagContainerType TmpContainerType = EPropertyBagContainerType::None;
+			Ar << TmpContainerType;
+
+			Bag.ContainerTypes.Add(TmpContainerType);
+		}
+		else
+		{
+			Ar << Bag.ContainerTypes;
+		}
 	}
 
 	bool bHasMetaData = false;
@@ -1106,7 +1206,7 @@ bool FPropertyBagPropertyDesc::IsClassType() const
 bool FPropertyBagPropertyDesc::CompatibleType(const FPropertyBagPropertyDesc& Other) const
 {
 	// Containers must match
-	if (ContainerType != Other.ContainerType)
+	if (ContainerTypes != Other.ContainerTypes)
 	{
 		return false;
 	}
@@ -1194,6 +1294,11 @@ void FInstancedPropertyBag::AddContainerProperty(const FName InName, const EProp
 	AddProperties({ FPropertyBagPropertyDesc(InName, InContainerType, InValueType, InValueTypeObject) });
 }
 
+void FInstancedPropertyBag::AddContainerProperty(const FName InName, const FPropertyBagContainerTypes InContainerTypes, const EPropertyBagPropertyType InValueType, UObject* InValueTypeObject)
+{
+	AddProperties({ FPropertyBagPropertyDesc(InName, InContainerTypes, InValueType, InValueTypeObject) });
+}
+
 void FInstancedPropertyBag::AddProperty(const FName InName, const FProperty* InSourceProperty)
 {
 	AddProperties({ FPropertyBagPropertyDesc(InName, InSourceProperty ) });
@@ -1594,7 +1699,7 @@ TValueOrError<FPropertyBagArrayRef, EPropertyBagResult> FInstancedPropertyBag::G
 	}
 	check(Desc->CachedProperty);
 
-	if (Desc->ContainerType != EPropertyBagContainerType::Array)
+	if (Desc->ContainerTypes.GetFirstContainerType() != EPropertyBagContainerType::Array)
 	{
 		return MakeError(EPropertyBagResult::TypeMismatch);
 	}
@@ -1617,7 +1722,7 @@ TValueOrError<const FPropertyBagArrayRef, EPropertyBagResult> FInstancedProperty
 	}
 	check(Desc->CachedProperty);
 
-	if (Desc->ContainerType != EPropertyBagContainerType::Array)
+	if (Desc->ContainerTypes.GetFirstContainerType() != EPropertyBagContainerType::Array)
 	{
 		return MakeError(EPropertyBagResult::TypeMismatch);
 	}
@@ -1866,6 +1971,49 @@ TValueOrError<UClass*, EPropertyBagResult> FPropertyBagArrayRef::GetValueClass(c
 	return MakeValue(Class);
 }
 
+TValueOrError<FPropertyBagArrayRef, EPropertyBagResult> FPropertyBagArrayRef::GetMutableNestedArrayRef(const int32 Index) const
+{
+	if (ValueDesc.ContainerTypes.GetFirstContainerType() != EPropertyBagContainerType::Array)
+	{
+		return MakeError(EPropertyBagResult::TypeMismatch);
+	}
+
+	check(ValueDesc.CachedProperty);
+
+	const FArrayProperty* ArrayProperty = CastField<FArrayProperty>(ValueDesc.CachedProperty);
+
+	// Get the array address
+	const void* Address = GetAddress(Index);
+	if (Address == nullptr)
+	{
+		return MakeError(EPropertyBagResult::PropertyNotFound);
+	}
+	
+	// And create a FPropertyBagArrayRef with the dummy desc and the element address
+	return MakeValue(FPropertyBagArrayRef(ValueDesc, Address));
+}
+
+TValueOrError<const FPropertyBagArrayRef, EPropertyBagResult> FPropertyBagArrayRef::GetNestedArrayRef(const int32 Index) const
+{
+	if (ValueDesc.ContainerTypes.GetFirstContainerType() != EPropertyBagContainerType::Array)
+	{
+		return MakeError(EPropertyBagResult::TypeMismatch);
+	}
+
+	check(ValueDesc.CachedProperty);
+
+	const FArrayProperty* ArrayProperty = CastField<FArrayProperty>(ValueDesc.CachedProperty);
+
+	// Get the array address
+	const void* Address = GetAddress(Index);
+	if (Address == nullptr)
+	{
+		return MakeError(EPropertyBagResult::PropertyNotFound);
+	}
+	
+	// And create a FPropertyBagArrayRef with the dummy desc and the element address
+	return MakeValue(FPropertyBagArrayRef(ValueDesc, Address));
+}
 
 EPropertyBagResult FPropertyBagArrayRef::SetValueBool(const int32 Index, const bool bInValue)
 {
diff --git a/Engine/Plugins/Experimental/StructUtils/Source/StructUtils/Public/PropertyBag.h b/Engine/Plugins/Experimental/StructUtils/Source/StructUtils/Public/PropertyBag.h
--- a/Engine/Plugins/Experimental/StructUtils/Source/StructUtils/Public/PropertyBag.h
+++ b/Engine/Plugins/Experimental/StructUtils/Source/StructUtils/Public/PropertyBag.h
@@ -41,6 +41,112 @@ enum class EPropertyBagContainerType : uint8
 	Count UMETA(Hidden)
 };
 
+/** Helper to manage container types, with nested container support. */
+USTRUCT()
+struct STRUCTUTILS_API FPropertyBagContainerTypes
+{
+	GENERATED_BODY()
+
+	FPropertyBagContainerTypes()
+		: Types{ EPropertyBagContainerType::None, EPropertyBagContainerType::None }
+		, NumContainers(0)
+	{}
+
+	explicit FPropertyBagContainerTypes(EPropertyBagContainerType ContainerType)
+		: Types{ ContainerType, EPropertyBagContainerType::None }
+		, NumContainers(1)
+	{}
+
+	FPropertyBagContainerTypes(const std::initializer_list<EPropertyBagContainerType>& InTypes)
+		: Types{ EPropertyBagContainerType::None, EPropertyBagContainerType::None }
+		, NumContainers(0)
+	{
+		for (EPropertyBagContainerType ContainerType : InTypes)
+		{
+			Add(ContainerType);
+		}
+	}
+
+	bool Add(EPropertyBagContainerType PropertyBagContainerType)
+	{
+		if (ensure(NumContainers < MaxNestedTypes))
+		{
+			Types[NumContainers] = PropertyBagContainerType;
+			NumContainers++;
+			return true;
+		}
+
+		return false;
+	}
+
+	void Reset()
+	{
+		for (auto& Type : Types)
+		{
+			Type = EPropertyBagContainerType::None;
+		}
+		NumContainers = 0;
+	}
+
+	bool IsEmpty() const
+	{
+		return NumContainers == 0;
+	}
+
+	uint32 Num() const
+	{
+		return NumContainers;
+	}
+
+	bool CanAdd() const
+	{
+		return NumContainers < MaxNestedTypes;
+	}
+
+	EPropertyBagContainerType GetFirstContainerType() const
+	{
+		return NumContainers > 0 ? Types[0] : EPropertyBagContainerType::None;
+	}
+
+	EPropertyBagContainerType operator[] (int32 Index) const
+	{
+		return ensure(Index < NumContainers) ? Types[Index] : EPropertyBagContainerType::None;
+	}
+
+	EPropertyBagContainerType PopHead();
+
+	void Serialize(FArchive& Ar);
+
+	friend FArchive& operator<<(FArchive& Ar, FPropertyBagContainerTypes& ContainerTypesData)
+	{
+		ContainerTypesData.Serialize(Ar);
+		return Ar;
+	}
+
+	bool operator == (const FPropertyBagContainerTypes& Other) const;
+
+	FORCEINLINE bool operator !=(const FPropertyBagContainerTypes& Other) const
+	{
+		return !(Other == *this);
+	}
+
+	friend FORCEINLINE uint32 GetTypeHash(const FPropertyBagContainerTypes& PropertyBagContainerTypes)
+	{
+		return HashCombine(GetTypeHash(PropertyBagContainerTypes.Types), GetTypeHash(PropertyBagContainerTypes.NumContainers));
+	}
+
+	EPropertyBagContainerType* begin() { return &Types[0]; }
+	const EPropertyBagContainerType* begin() const { return &Types[0]; }
+	EPropertyBagContainerType* end()  { return &Types[NumContainers]; }
+	const EPropertyBagContainerType* end() const { return &Types[NumContainers]; }
+
+protected:
+	static constexpr uint8 MaxNestedTypes = 2;
+
+	EPropertyBagContainerType Types[MaxNestedTypes];
+	uint8 NumContainers = 0;
+};
+
 /** Getter and setter result code. */
 UENUM()
 enum class EPropertyBagResult : uint8
@@ -116,10 +222,26 @@ struct STRUCTUTILS_API FPropertyBagPropertyDesc
 		: ValueTypeObject(InValueTypeObject)
 		, Name(InName)
 		, ValueType(InValueType)
-		, ContainerType(InContainerType)
+		, ContainerTypes(InContainerType)
+	{
+	}
+
+	FPropertyBagPropertyDesc(const FName InName, const FPropertyBagContainerTypes& InNestedContainers, const EPropertyBagPropertyType InValueType, UObject* InValueTypeObject = nullptr)
+		: ValueTypeObject(InValueTypeObject)
+		, Name(InName)
+		, ValueType(InValueType)
+		, ContainerTypes(InNestedContainers)
 	{
 	}
 
+	// @todo: Remove these when ContainerType deprecation is removed.
+	PRAGMA_DISABLE_DEPRECATION_WARNINGS
+	FPropertyBagPropertyDesc(const FPropertyBagPropertyDesc&) = default;
+	FPropertyBagPropertyDesc(FPropertyBagPropertyDesc&&) noexcept = default;
+	FPropertyBagPropertyDesc& operator= (const FPropertyBagPropertyDesc&) = default;
+	FPropertyBagPropertyDesc& operator= (FPropertyBagPropertyDesc&& other) noexcept = default;
+	PRAGMA_ENABLE_DEPRECATION_WARNINGS
+
 	/** @return true if the two descriptors have the same type. Object types are compatible if Other can be cast to this type. */
 	bool CompatibleType(const FPropertyBagPropertyDesc& Other) const;
 
@@ -151,9 +273,13 @@ struct STRUCTUTILS_API FPropertyBagPropertyDesc
 	UPROPERTY(EditAnywhere, Category="Default")
 	EPropertyBagPropertyType ValueType = EPropertyBagPropertyType::None;
 
+	/** Type of the container described by this property. */
+	UE_DEPRECATED(5.3, "ContainerType has been deprecated as of 5.3. Please use ContainerTypes instead.")
+	FPropertyBagContainerTypes ContainerType;
+
 	/** Type of the container described by this property. */
 	UPROPERTY(EditAnywhere, Category="Default")
-	EPropertyBagContainerType ContainerType = EPropertyBagContainerType::None;
+	FPropertyBagContainerTypes ContainerTypes;
 
 #if WITH_EDITORONLY_DATA
 	/** Editor-only meta data for CachedProperty */
@@ -272,6 +398,15 @@ struct STRUCTUTILS_API FInstancedPropertyBag
 	 */
 	void AddContainerProperty(const FName InName, const EPropertyBagContainerType InContainerType, const EPropertyBagPropertyType InValueType, UObject* InValueTypeObject = nullptr);
 
+	/**
+	 * Adds a new container property to the bag. If property of same name already exists, it will be replaced with the new type.
+	 * @param InName Name of the new property
+	 * @param InContainerTypes List of (optionally nested) containers to create
+	 * @param InValueType Type of the new property
+	 * @param InValueTypeObject Type object (for struct, class, enum) of the new property
+	 */
+	void AddContainerProperty(const FName InName, const FPropertyBagContainerTypes InContainerTypes, const EPropertyBagPropertyType InValueType, UObject* InValueTypeObject);
+
 	/**
 	 * Adds a new property to the bag. Property type duplicated from source property to. If property of same name already exists, it will be replaced with the new type.
 	 * @param Descs Descriptors of new properties to add.
@@ -532,7 +667,8 @@ public:
 		ValueDesc.ValueType = InDesc.ValueType;
 		ValueDesc.ValueTypeObject = InDesc.ValueTypeObject;
 		ValueDesc.CachedProperty = ArrayProperty->Inner;
-		ValueDesc.ContainerType = EPropertyBagContainerType::None;
+		ValueDesc.ContainerTypes = InDesc.ContainerTypes;
+		ValueDesc.ContainerTypes.PopHead();
 	}
 
 	/**
@@ -606,6 +742,20 @@ public:
 		return MakeError(EPropertyBagResult::TypeMismatch);
 	}
 
+	/**
+     * Returns helper class to modify and access a nested array (mutable version).
+     * Note: Note: The array reference is not valid after the layout of the referenced property bag has changed!
+     * @returns helper class to modify and access arrays
+    */
+	TValueOrError<FPropertyBagArrayRef, EPropertyBagResult> GetMutableNestedArrayRef(const int32 Index = 0) const;
+
+	/**
+     * Returns helper class to access a nested array (const version).
+     * Note: Note: The array reference is not valid after the layout of the referenced property bag has changed!
+     * @returns helper class to access arrays
+    */
+	TValueOrError<const FPropertyBagArrayRef, EPropertyBagResult> GetNestedArrayRef(const int32 Index = 0) const;
+
 	/**
 	 * Value Setters. A property must exists in that bag before it can be set.  
 	 * Numeric types (bool, int32, int64, float, double) support type conversion.
diff --git a/Engine/Plugins/Experimental/StructUtils/Source/StructUtilsEditor/Private/Customizations/PropertyBagDetails.cpp b/Engine/Plugins/Experimental/StructUtils/Source/StructUtilsEditor/Private/Customizations/PropertyBagDetails.cpp
--- a/Engine/Plugins/Experimental/StructUtils/Source/StructUtilsEditor/Private/Customizations/PropertyBagDetails.cpp
+++ b/Engine/Plugins/Experimental/StructUtils/Source/StructUtilsEditor/Private/Customizations/PropertyBagDetails.cpp
@@ -35,7 +35,7 @@ bool IsScriptStruct(const TSharedPtr<IPropertyHandle>& PropertyHandle)
 	}
 
 	FStructProperty* StructProperty = CastField<FStructProperty>(PropertyHandle->GetProperty());
-	return StructProperty && StructProperty->Struct == TBaseStructure<T>::Get();
+	return StructProperty && StructProperty->Struct->IsA(TBaseStructure<T>::Get()->GetClass());
 }
 
 /** @return true if the property is one of the known missing types. */
@@ -219,7 +219,9 @@ FEdGraphPinType GetPropertyDescAsPin(const FPropertyBagPropertyDesc& Desc)
 	PinType.PinSubCategory = NAME_None;
 
 	// Container type
-	switch (Desc.ContainerType)
+	//@todo: Handle nested containers in property selection.
+	const EPropertyBagContainerType ContainerType = Desc.ContainerTypes.GetFirstContainerType();
+	switch (ContainerType)
 	{
 	case EPropertyBagContainerType::Array:
 		PinType.ContainerType = EPinContainerType::Array;
@@ -299,14 +301,23 @@ void SetPropertyDescFromPin(FPropertyBagPropertyDesc& Desc, const FEdGraphPinTyp
 	const UPropertyBagSchema* Schema = GetDefault<UPropertyBagSchema>();
 	check(Schema);
 
-	// Container type
+	// remove any existing containers
+	Desc.ContainerTypes.Reset();
+
+	// Fill Container types, if any
 	switch (PinType.ContainerType)
 	{
 	case EPinContainerType::Array:
-		Desc.ContainerType = EPropertyBagContainerType::Array;
+		Desc.ContainerTypes.Add(EPropertyBagContainerType::Array);
+		break;
+	case EPinContainerType::Set:
+		ensureMsgf(false, TEXT("Unsuported container type [Set] "));
+		break;
+	case EPinContainerType::Map:
+		ensureMsgf(false, TEXT("Unsuported container type [Map] "));
 		break;
 	default:
-		Desc.ContainerType = EPropertyBagContainerType::None;
+		break;
 	}
 	
 	// Value type
@@ -478,6 +489,19 @@ void FPropertyBagInstanceDataDetails::OnChildRowAdded(IDetailPropertyRow& ChildR
 
 	TSharedPtr<IPropertyHandle> ChildPropertyHandle = ChildRow.GetPropertyHandle();
 	
+	bool bSupportedType = true;
+
+	TArray<FPropertyBagPropertyDesc> PropertyDescs = UE::StructUtils::Private::GetCommonPropertyDescs(BagStructProperty);
+	const FProperty* Property = ChildPropertyHandle->GetProperty();
+	if (FPropertyBagPropertyDesc* Desc = PropertyDescs.FindByPredicate([Property](const FPropertyBagPropertyDesc& Desc){ return Desc.CachedProperty == Property; }))
+	{
+		if (Desc->ContainerTypes.Num() > 1)
+		{
+			bSupportedType = false;
+			ChildRow.IsEnabled(false);
+		}
+	}
+
 	if (!bFixedLayout)
 	{
 		// Inline editable name
@@ -552,10 +576,15 @@ void FPropertyBagInstanceDataDetails::OnChildRowAdded(IDetailPropertyRow& ChildR
 				.HeightOverride(12)
 				[
 					SNew(SImage)
-					.ToolTipText(LOCTEXT("MissingType", "The property is missing type. The Struct, Enum, or Object may have been removed."))
-					.Visibility_Lambda([ChildPropertyHandle]()
+					.ToolTipText_Lambda([ChildPropertyHandle, bSupportedType]()
+					{
+						return !bSupportedType
+							? LOCTEXT("UnsupportedType", "This property type is not supported in the property bag UI.")
+							: LOCTEXT("MissingType", "The property is missing type. The Struct, Enum, or Object may have been removed.");
+					})
+					.Visibility_Lambda([ChildPropertyHandle, bSupportedType]()
 					{
-						return UE::StructUtils::Private::HasMissingType(ChildPropertyHandle) ? EVisibility::Visible : EVisibility::Collapsed;
+						return UE::StructUtils::Private::HasMissingType(ChildPropertyHandle) || !bSupportedType ? EVisibility::Visible : EVisibility::Collapsed;
 					})
 					.Image(FAppStyle::GetBrush("Icons.Error"))
 				]
diff --git a/Engine/Plugins/Experimental/StructUtils/Source/StructUtilsTestSuite/Private/PropertyBagTest.cpp b/Engine/Plugins/Experimental/StructUtils/Source/StructUtilsTestSuite/Private/PropertyBagTest.cpp
--- a/Engine/Plugins/Experimental/StructUtils/Source/StructUtilsTestSuite/Private/PropertyBagTest.cpp
+++ b/Engine/Plugins/Experimental/StructUtils/Source/StructUtilsTestSuite/Private/PropertyBagTest.cpp
@@ -255,6 +255,109 @@ struct FTest_Enum : FAITestBase
 };
 IMPLEMENT_AI_INSTANT_TEST(FTest_Enum, "System.StructUtils.PropertyBag.Enum");
 
+struct FTest_NestedArray : FAITestBase
+{
+	virtual bool InstantTest() override
+	{
+		static const FName NestedInt32ArrayPropName(TEXT("NestedInt32ArrayProp"));
+		static const TArray<TArray<int32>> NestedInt32ArrayTestValue = { { 1, 2, 3} };
+
+		FInstancedPropertyBag Bag;
+
+		// Set properties
+		{
+			Bag.AddContainerProperty(NestedInt32ArrayPropName, { EPropertyBagContainerType::Array, EPropertyBagContainerType::Array }, EPropertyBagPropertyType::Int32, nullptr);
+
+			AITEST_TRUE(TEXT("Missing Float Array property in the Bag."), Bag.FindPropertyDescByName(NestedInt32ArrayPropName) != nullptr);
+		}
+
+		// Check Default Value
+		{
+			TValueOrError<const FPropertyBagArrayRef, EPropertyBagResult> NestedInt32ArrayDefaultResult = Bag.GetArrayRef(NestedInt32ArrayPropName);
+
+			AITEST_TRUE(TEXT("Bag getting Nested Int32 Array default value should succeed."), NestedInt32ArrayDefaultResult.IsValid());
+
+			if (NestedInt32ArrayDefaultResult.IsValid())
+			{
+				AITEST_TRUE(TEXT("Bag Nested Int32 Array default value incorrect size"), NestedInt32ArrayDefaultResult.GetValue().Num() == 0);
+			}
+		}
+
+		// Set Nested Array values using FPropertyBagArrayRef interface
+		{
+			TValueOrError<FPropertyBagArrayRef, EPropertyBagResult> NestedInt32ArrayMutable = Bag.GetMutableArrayRef(NestedInt32ArrayPropName);
+
+			AITEST_TRUE(TEXT("Getting PropertyBag Nested Int32 Array should succeed."), NestedInt32ArrayMutable.IsValid());
+
+			FPropertyBagArrayRef& PropertyBagArrayRef = NestedInt32ArrayMutable.GetValue();
+
+			const int32 NumArrays = NestedInt32ArrayTestValue.Num();
+			PropertyBagArrayRef.AddValues(NumArrays);
+
+			for (int32 n = 0; n < NumArrays; n++)
+			{
+				TValueOrError<FPropertyBagArrayRef, EPropertyBagResult> InnerArrayTestResult = PropertyBagArrayRef.GetMutableNestedArrayRef(n);
+					
+				AITEST_TRUE(TEXT("Getting PropertyBag Nested Inner Int32 Array should succeed."), InnerArrayTestResult.IsValid());
+					
+				FPropertyBagArrayRef& InnerPropertyBagArrayRef = InnerArrayTestResult.GetValue();
+
+				const int32 NumElem = NestedInt32ArrayTestValue[n].Num();
+				InnerPropertyBagArrayRef.AddUninitializedValues(NumElem);
+
+				for (int32 i = 0; i < NumElem; i++)
+				{
+					const int32& value = NestedInt32ArrayTestValue[n][i];
+
+					AITEST_TRUE(TEXT("Setting value to Nested Inner Array property failed."), InnerPropertyBagArrayRef.SetValueInt32(i, value) == EPropertyBagResult::Success);
+				}
+			}
+		}
+
+		// Test Nested Array Values using FPropertyBagArrayRef interface
+		{
+			TValueOrError<const FPropertyBagArrayRef, EPropertyBagResult> NestedInt32ArrayTestResult = Bag.GetArrayRef(NestedInt32ArrayPropName);
+
+			AITEST_TRUE(TEXT("Getting PropertyBag Nested Int32 Array should succeed."), NestedInt32ArrayTestResult.IsValid());
+
+			if (NestedInt32ArrayTestResult.IsValid())
+			{
+				const FPropertyBagArrayRef& NestedPropertyBagArrayRef = NestedInt32ArrayTestResult.GetValue();
+
+				const int32 NumArrays = NestedPropertyBagArrayRef.Num();
+				const int32 NumTestArrays = NestedInt32ArrayTestValue.Num();
+
+				AITEST_TRUE(TEXT("Bag [%s] Nested Int32 Array Num value mismatch."), NumArrays == NumTestArrays);
+
+				for (int32 n = 0; n < NumArrays; ++n)
+				{
+					TValueOrError<const FPropertyBagArrayRef, EPropertyBagResult> InnerArrayTestResult = NestedPropertyBagArrayRef.GetNestedArrayRef();
+
+					AITEST_TRUE(TEXT("Getting PropertyBag Nested Inner Int32 Array should succeed."), InnerArrayTestResult.IsValid());
+
+					const FPropertyBagArrayRef& PropertyBagArrayRef = InnerArrayTestResult.GetValue();
+
+					const int32 NumElem = InnerArrayTestResult.GetValue().Num();
+					for (int32 i = 0; i < NumElem; i++)
+					{
+						const TValueOrError<int32, EPropertyBagResult> Int32Res = PropertyBagArrayRef.GetValueInt32(i);
+
+						AITEST_TRUE(TEXT("Getting Nested Array Element should succeed."), Int32Res.IsValid());
+
+						if (Int32Res.IsValid())
+						{
+							AITEST_TRUE(TEXT("Nested Arrauy test value mismatch."), Int32Res.GetValue() == NestedInt32ArrayTestValue[n][i]);
+						}
+					}
+				}
+			}
+		}
+
+		return true;
+	}
+};
+IMPLEMENT_AI_INSTANT_TEST(FTest_NestedArray, "System.StructUtils.PropertyBag.NestedArray");
+
 struct FTest_GC : FAITestBase
 {
 	virtual bool InstantTest() override
diff --git a/Engine/Plugins/MovieScene/MovieRenderPipeline/Source/MovieRenderPipelineCore/Private/MovieJobVariableAssignmentContainer.cpp b/Engine/Plugins/MovieScene/MovieRenderPipeline/Source/MovieRenderPipelineCore/Private/MovieJobVariableAssignmentContainer.cpp
--- a/Engine/Plugins/MovieScene/MovieRenderPipeline/Source/MovieRenderPipelineCore/Private/MovieJobVariableAssignmentContainer.cpp
+++ b/Engine/Plugins/MovieScene/MovieRenderPipeline/Source/MovieRenderPipelineCore/Private/MovieJobVariableAssignmentContainer.cpp
@@ -212,7 +212,7 @@ EMovieGraphContainerType UMovieJobVariableAssignmentContainer::GetValueContainer
 {
 	if (const FPropertyBagPropertyDesc* Desc = Value.FindPropertyDescByName(PropertyName))
 	{
-		return static_cast<EMovieGraphContainerType>(Desc->ContainerType);
+		return static_cast<EMovieGraphContainerType>(Desc->ContainerTypes.GetFirstContainerType());
 	}
 
 	return EMovieGraphContainerType::None;
@@ -254,7 +254,7 @@ bool UMovieJobVariableAssignmentContainer::GenerateVariableOverride(const UMovie
 	FPropertyBagPropertyDesc NewProperty = FPropertyBagPropertyDesc();
 	NewProperty.ValueType = static_cast<EPropertyBagPropertyType>(InGraphVariable->GetValueType());
 	NewProperty.ValueTypeObject = InGraphVariable->GetValueTypeObject();
-	NewProperty.ContainerType = static_cast<EPropertyBagContainerType>(InGraphVariable->GetValueContainerType());
+	NewProperty.ContainerTypes = { static_cast<EPropertyBagContainerType>(InGraphVariable->GetValueContainerType()) };
 	NewProperty.Name = FName(InGraphVariable->Name);
 	NewProperty.ID = FGuid::NewGuid();
 #if WITH_EDITOR
@@ -384,10 +384,10 @@ void UMovieJobVariableAssignmentContainer::UpdateGraphVariableOverrides()
 				Desc.ValueTypeObject = Variable->GetValueTypeObject();
 			}
 
-			if ((uint8)Desc.ContainerType != (uint8)(Variable->GetValueContainerType()))
+			if ((uint8)Desc.ContainerTypes.GetFirstContainerType() != (uint8)(Variable->GetValueContainerType()))
 			{
 				bNeedsToRegenerate = true;
-				Desc.ContainerType = static_cast<EPropertyBagContainerType>(Variable->GetValueContainerType());
+				Desc.ContainerTypes = { static_cast<EPropertyBagContainerType>(Variable->GetValueContainerType()) };
 			}
 		}
 	}
diff --git a/Engine/Plugins/MovieScene/MovieRenderPipeline/Source/MovieRenderPipelineCore/Public/Graph/MovieGraphValueContainer.cpp b/Engine/Plugins/MovieScene/MovieRenderPipeline/Source/MovieRenderPipelineCore/Public/Graph/MovieGraphValueContainer.cpp
--- a/Engine/Plugins/MovieScene/MovieRenderPipeline/Source/MovieRenderPipelineCore/Public/Graph/MovieGraphValueContainer.cpp
+++ b/Engine/Plugins/MovieScene/MovieRenderPipeline/Source/MovieRenderPipelineCore/Public/Graph/MovieGraphValueContainer.cpp
@@ -223,7 +223,7 @@ EMovieGraphContainerType UMovieGraphValueContainer::GetValueContainerType() cons
 {
 	if (const FPropertyBagPropertyDesc* Desc = Value.FindPropertyDescByName(PropertyBagDefaultPropertyName))
 	{
-		return static_cast<EMovieGraphContainerType>(Desc->ContainerType);
+		return static_cast<EMovieGraphContainerType>(Desc->ContainerTypes.GetFirstContainerType());
 	}
 
 	return EMovieGraphContainerType::None;
@@ -234,7 +234,7 @@ void UMovieGraphValueContainer::SetValueContainerType(EMovieGraphContainerType C
 	if (const FPropertyBagPropertyDesc* Desc = Value.FindPropertyDescByName(PropertyBagDefaultPropertyName))
 	{
 		FPropertyBagPropertyDesc NewDesc(*Desc);
-		NewDesc.ContainerType = static_cast<EPropertyBagContainerType>(ContainerType);
+		NewDesc.ContainerTypes = { static_cast<EPropertyBagContainerType>(ContainerType) };
 
 		const UPropertyBag* NewPropBag = UPropertyBag::GetOrCreateFromDescs({NewDesc});
 		Value.MigrateToNewBagStruct(NewPropBag);
