CIS - Fixed static analisys warning at PropertyBag.cpp and a wrong deprecation at PropertyBag.h

diff --git a/Engine/Plugins/Experimental/StructUtils/Source/StructUtils/Private/PropertyBag.cpp b/Engine/Plugins/Experimental/StructUtils/Source/StructUtils/Private/PropertyBag.cpp
--- a/Engine/Plugins/Experimental/StructUtils/Source/StructUtils/Private/PropertyBag.cpp
+++ b/Engine/Plugins/Experimental/StructUtils/Source/StructUtils/Private/PropertyBag.cpp
@@ -1048,9 +1048,9 @@ EPropertyBagContainerType FPropertyBagContainerTypes::PopHead()
 	{
 		Head = Types[0];
 
-		for (uint32 i = 1; i < NumContainers; ++i)
+		for (uint32 i = 0; i < (NumContainers - 1u); ++i)
 		{
-			Types[i - 1] = Types[i];
+			Types[i] = Types[i + 1];
 		}
 		Types[NumContainers - 1] = EPropertyBagContainerType::None;
 		NumContainers = NumContainers - 1;
diff --git a/Engine/Plugins/Experimental/StructUtils/Source/StructUtils/Public/PropertyBag.h b/Engine/Plugins/Experimental/StructUtils/Source/StructUtils/Public/PropertyBag.h
--- a/Engine/Plugins/Experimental/StructUtils/Source/StructUtils/Public/PropertyBag.h
+++ b/Engine/Plugins/Experimental/StructUtils/Source/StructUtils/Public/PropertyBag.h
@@ -273,9 +273,11 @@ struct STRUCTUTILS_API FPropertyBagPropertyDesc
 	UPROPERTY(EditAnywhere, Category="Default")
 	EPropertyBagPropertyType ValueType = EPropertyBagPropertyType::None;
 
-	/** Type of the container described by this property. */
-	UE_DEPRECATED(5.3, "ContainerType has been deprecated as of 5.3. Please use ContainerTypes instead.")
-	FPropertyBagContainerTypes ContainerType;
+#if WITH_EDITORONLY_DATA
+	/** DEPRECATED 5.3 */
+	UPROPERTY(meta = (DeprecatedProperty, DeprecationMessage = "ContainerType has been deprecated as of 5.3. Please use ContainerTypes instead."))
+	FPropertyBagContainerTypes ContainerType_DEPRECATED; 
+#endif
 
 	/** Type of the container described by this property. */
 	UPROPERTY(EditAnywhere, Category="Default")
