CIS - Fixed bad deprecation of PropertyBag generating ensures at packaging

diff --git a/Engine/Plugins/Experimental/StructUtils/Source/StructUtils/Private/PropertyBag.cpp b/Engine/Plugins/Experimental/StructUtils/Source/StructUtils/Private/PropertyBag.cpp
--- a/Engine/Plugins/Experimental/StructUtils/Source/StructUtils/Private/PropertyBag.cpp
+++ b/Engine/Plugins/Experimental/StructUtils/Source/StructUtils/Private/PropertyBag.cpp
@@ -1129,7 +1129,10 @@ static FArchive& operator<<(FArchive& Ar, FPropertyBagPropertyDesc& Bag)
 			EPropertyBagContainerType TmpContainerType = EPropertyBagContainerType::None;
 			Ar << TmpContainerType;
 
-			Bag.ContainerTypes.Add(TmpContainerType);
+			if (TmpContainerType != EPropertyBagContainerType::None)
+			{
+				Bag.ContainerTypes.Add(TmpContainerType);
+			}
 		}
 		else
 		{
diff --git a/Engine/Plugins/Experimental/StructUtils/Source/StructUtils/Public/PropertyBag.h b/Engine/Plugins/Experimental/StructUtils/Source/StructUtils/Public/PropertyBag.h
--- a/Engine/Plugins/Experimental/StructUtils/Source/StructUtils/Public/PropertyBag.h
+++ b/Engine/Plugins/Experimental/StructUtils/Source/StructUtils/Public/PropertyBag.h
@@ -53,9 +53,14 @@ struct STRUCTUTILS_API FPropertyBagContainerTypes
 	{}
 
 	explicit FPropertyBagContainerTypes(EPropertyBagContainerType ContainerType)
-		: Types{ ContainerType, EPropertyBagContainerType::None }
-		, NumContainers(1)
-	{}
+		: Types{ EPropertyBagContainerType::None, EPropertyBagContainerType::None }
+		, NumContainers(0)
+	{
+		if (ContainerType != EPropertyBagContainerType::None)
+		{
+			Add(ContainerType);
+		}
+	}
 
 	FPropertyBagContainerTypes(const std::initializer_list<EPropertyBagContainerType>& InTypes)
 		: Types{ EPropertyBagContainerType::None, EPropertyBagContainerType::None }
@@ -63,7 +68,10 @@ struct STRUCTUTILS_API FPropertyBagContainerTypes
 	{
 		for (EPropertyBagContainerType ContainerType : InTypes)
 		{
-			Add(ContainerType);
+			if (ContainerType != EPropertyBagContainerType::None)
+			{
+				Add(ContainerType);
+			}
 		}
 	}
 
@@ -234,14 +242,6 @@ struct STRUCTUTILS_API FPropertyBagPropertyDesc
 	{
 	}
 
-	// @todo: Remove these when ContainerType deprecation is removed.
-	PRAGMA_DISABLE_DEPRECATION_WARNINGS
-	FPropertyBagPropertyDesc(const FPropertyBagPropertyDesc&) = default;
-	FPropertyBagPropertyDesc(FPropertyBagPropertyDesc&&) noexcept = default;
-	FPropertyBagPropertyDesc& operator= (const FPropertyBagPropertyDesc&) = default;
-	FPropertyBagPropertyDesc& operator= (FPropertyBagPropertyDesc&& other) noexcept = default;
-	PRAGMA_ENABLE_DEPRECATION_WARNINGS
-
 	/** @return true if the two descriptors have the same type. Object types are compatible if Other can be cast to this type. */
 	bool CompatibleType(const FPropertyBagPropertyDesc& Other) const;
 
@@ -273,12 +273,6 @@ struct STRUCTUTILS_API FPropertyBagPropertyDesc
 	UPROPERTY(EditAnywhere, Category="Default")
 	EPropertyBagPropertyType ValueType = EPropertyBagPropertyType::None;
 
-#if WITH_EDITORONLY_DATA
-	/** DEPRECATED 5.3 */
-	UPROPERTY(meta = (DeprecatedProperty, DeprecationMessage = "ContainerType has been deprecated as of 5.3. Please use ContainerTypes instead."))
-	FPropertyBagContainerTypes ContainerType_DEPRECATED; 
-#endif
-
 	/** Type of the container described by this property. */
 	UPROPERTY(EditAnywhere, Category="Default")
 	FPropertyBagContainerTypes ContainerTypes;
