Fix for a crash when hovering a dropdown without any instance struct.
Updating engine change for 25.10 Original Change

Show ToolTip text for instance structs instead of DisplayName text.

diff --git a/Engine/Plugins/Experimental/StructUtils/Source/StructUtilsEditor/Private/Customizations/InstancedStructDetails.cpp b/Engine/Plugins/Experimental/StructUtils/Source/StructUtilsEditor/Private/Customizations/InstancedStructDetails.cpp
--- a/Engine/Plugins/Experimental/StructUtils/Source/StructUtilsEditor/Private/Customizations/InstancedStructDetails.cpp
+++ b/Engine/Plugins/Experimental/StructUtils/Source/StructUtilsEditor/Private/Customizations/InstancedStructDetails.cpp
@@ -449,7 +449,7 @@ FText FInstancedStructDetails::GetTooltipText() const
 	const UScriptStruct* CommonStruct = nullptr;
 	const FPropertyAccess::Result Result = UE::StructUtils::Private::GetCommonScriptStruct(StructProperty, CommonStruct);
 	
-	if (Result == FPropertyAccess::Success)
+	if (CommonStruct && Result == FPropertyAccess::Success)
 	{
 		return CommonStruct->GetToolTipText();
 	}
