PropertyBag : Added a check to avoid adding empty containers in the Add Functiion
Created a new PropertyBag test to do minimal check of the Containers struct
Attempt to remove a static analiys error about out of bounds, by adding a check as a hint for the compiler

diff --git a/Engine/Plugins/Experimental/StructUtils/Source/StructUtils/Private/PropertyBag.cpp b/Engine/Plugins/Experimental/StructUtils/Source/StructUtils/Private/PropertyBag.cpp
--- a/Engine/Plugins/Experimental/StructUtils/Source/StructUtils/Private/PropertyBag.cpp
+++ b/Engine/Plugins/Experimental/StructUtils/Source/StructUtils/Private/PropertyBag.cpp
@@ -1048,12 +1048,15 @@ EPropertyBagContainerType FPropertyBagContainerTypes::PopHead()
 	{
 		Head = Types[0];
 
-		for (uint32 i = 0; i < (NumContainers - 1u); ++i)
+		const uint32 NewNumContainers = NumContainers - 1u;
+		check(NewNumContainers < UE_ARRAY_COUNT(Types))
+
+		for (uint32 i = 0; i < NewNumContainers; ++i)
 		{
 			Types[i] = Types[i + 1];
 		}
-		Types[NumContainers - 1] = EPropertyBagContainerType::None;
-		NumContainers = NumContainers - 1;
+		Types[NewNumContainers] = EPropertyBagContainerType::None;
+		NumContainers = NewNumContainers;
 	}
 
 	return Head;
diff --git a/Engine/Plugins/Experimental/StructUtils/Source/StructUtils/Public/PropertyBag.h b/Engine/Plugins/Experimental/StructUtils/Source/StructUtils/Public/PropertyBag.h
--- a/Engine/Plugins/Experimental/StructUtils/Source/StructUtils/Public/PropertyBag.h
+++ b/Engine/Plugins/Experimental/StructUtils/Source/StructUtils/Public/PropertyBag.h
@@ -79,9 +79,13 @@ struct STRUCTUTILS_API FPropertyBagContainerTypes
 	{
 		if (ensure(NumContainers < MaxNestedTypes))
 		{
-			Types[NumContainers] = PropertyBagContainerType;
-			NumContainers++;
-			return true;
+			if (PropertyBagContainerType != EPropertyBagContainerType::None)
+			{
+				Types[NumContainers] = PropertyBagContainerType;
+				NumContainers++;
+
+				return true;
+			}
 		}
 
 		return false;
diff --git a/Engine/Plugins/Experimental/StructUtils/Source/StructUtilsTestSuite/Private/PropertyBagTest.cpp b/Engine/Plugins/Experimental/StructUtils/Source/StructUtilsTestSuite/Private/PropertyBagTest.cpp
--- a/Engine/Plugins/Experimental/StructUtils/Source/StructUtilsTestSuite/Private/PropertyBagTest.cpp
+++ b/Engine/Plugins/Experimental/StructUtils/Source/StructUtilsTestSuite/Private/PropertyBagTest.cpp
@@ -255,6 +255,35 @@ struct FTest_Enum : FAITestBase
 };
 IMPLEMENT_AI_INSTANT_TEST(FTest_Enum, "System.StructUtils.PropertyBag.Enum");
 
+struct FTest_ContainerTypes : FAITestBase
+{
+	virtual bool InstantTest() override
+	{
+		FPropertyBagContainerTypes Container = { EPropertyBagContainerType::None, EPropertyBagContainerType::None };
+		AITEST_TRUE(TEXT("Invalid Num Containers after creation."), Container.Num() == 0);
+		AITEST_TRUE(TEXT("Invalid First Container type after creation."), Container.GetFirstContainerType() == EPropertyBagContainerType::None);
+
+		Container.Add(EPropertyBagContainerType::Array);
+		AITEST_TRUE(TEXT("Invalid num containers"), Container.Num() == 1);
+		AITEST_TRUE(TEXT("Invalid First Container type."), Container.GetFirstContainerType() == EPropertyBagContainerType::Array);
+
+		Container.Add(EPropertyBagContainerType::Array);
+		AITEST_TRUE(TEXT("Invalid num containers"), Container.Num() == 2);
+
+		Container.PopHead();
+		AITEST_TRUE(TEXT("Invalid num containers"), Container.Num() == 1);
+
+		Container.PopHead();
+		AITEST_TRUE(TEXT("Invalid num containers"), Container.Num() == 0);
+
+		Container.Add(EPropertyBagContainerType::None);
+		AITEST_TRUE(TEXT("Adding None sould not change Num containers"), Container.Num() == 0);
+
+		return true;
+	}
+};
+IMPLEMENT_AI_INSTANT_TEST(FTest_ContainerTypes, "System.StructUtils.PropertyBag.ContainerTypes");
+
 struct FTest_NestedArray : FAITestBase
 {
 	virtual bool InstantTest() override
