Fix for CIS Static analysis errors at PropertyBag

diff --git a/Engine/Plugins/Experimental/StructUtils/Source/StructUtils/Private/PropertyBag.cpp b/Engine/Plugins/Experimental/StructUtils/Source/StructUtils/Private/PropertyBag.cpp
--- a/Engine/Plugins/Experimental/StructUtils/Source/StructUtils/Private/PropertyBag.cpp
+++ b/Engine/Plugins/Experimental/StructUtils/Source/StructUtils/Private/PropertyBag.cpp
@@ -1046,17 +1046,16 @@ EPropertyBagContainerType FPropertyBagContainerTypes::PopHead()
 		
 	if (NumContainers > 0)
 	{
-		Head = Types[0];
+		Swap(Head, Types[0]);
 
-		const uint32 NewNumContainers = NumContainers - 1u;
-		check(NewNumContainers < UE_ARRAY_COUNT(Types))
-
-		for (uint32 i = 0; i < NewNumContainers; ++i)
+		uint8 Index = NumContainers - 1;
+		while (Index > 0)
 		{
-			Types[i] = Types[i + 1];
+			Types[Index - 1] = Types[Index];
+			Types[Index] = EPropertyBagContainerType::None;
+			Index--;
 		}
-		Types[NewNumContainers] = EPropertyBagContainerType::None;
-		NumContainers = NewNumContainers;
+		NumContainers--;
 	}
 
 	return Head;
diff --git a/Engine/Plugins/Experimental/StructUtils/Source/StructUtilsTestSuite/Private/PropertyBagTest.cpp b/Engine/Plugins/Experimental/StructUtils/Source/StructUtilsTestSuite/Private/PropertyBagTest.cpp
--- a/Engine/Plugins/Experimental/StructUtils/Source/StructUtilsTestSuite/Private/PropertyBagTest.cpp
+++ b/Engine/Plugins/Experimental/StructUtils/Source/StructUtilsTestSuite/Private/PropertyBagTest.cpp
@@ -270,11 +270,15 @@ struct FTest_ContainerTypes : FAITestBase
 		Container.Add(EPropertyBagContainerType::Array);
 		AITEST_TRUE(TEXT("Invalid num containers"), Container.Num() == 2);
 
-		Container.PopHead();
+		const EPropertyBagContainerType HeadContainerType1 = Container.PopHead();
 		AITEST_TRUE(TEXT("Invalid num containers"), Container.Num() == 1);
+		AITEST_TRUE(TEXT("Invalid extracted head container 1"), HeadContainerType1 == EPropertyBagContainerType::Array);
+		AITEST_TRUE(TEXT("Invalid first container type after removing Head 1"), Container.GetFirstContainerType() == EPropertyBagContainerType::Array);
 
-		Container.PopHead();
+		const EPropertyBagContainerType HeadContainerType2 = Container.PopHead();
 		AITEST_TRUE(TEXT("Invalid num containers"), Container.Num() == 0);
+		AITEST_TRUE(TEXT("Invalid extracted head container 2"), HeadContainerType2 == EPropertyBagContainerType::Array);
+		AITEST_TRUE(TEXT("Invalid first container type after removing Head 2"), Container.GetFirstContainerType() == EPropertyBagContainerType::None);
 
 		Container.Add(EPropertyBagContainerType::None);
 		AITEST_TRUE(TEXT("Adding None sould not change Num containers"), Container.Num() == 0);
