Fix potential stale delegate binding

#rnx

diff --git a/Engine/Plugins/Experimental/StructUtils/Source/StructUtilsEditor/Private/Customizations/InstancedStructDetails.cpp b/Engine/Plugins/Experimental/StructUtils/Source/StructUtilsEditor/Private/Customizations/InstancedStructDetails.cpp
--- a/Engine/Plugins/Experimental/StructUtils/Source/StructUtilsEditor/Private/Customizations/InstancedStructDetails.cpp
+++ b/Engine/Plugins/Experimental/StructUtils/Source/StructUtilsEditor/Private/Customizations/InstancedStructDetails.cpp
@@ -491,7 +491,7 @@ TSharedRef<SWidget> FInstancedStructDetails::GenerateStructPicker()
 	Options.DisplayMode = bShowTreeView ? EStructViewerDisplayMode::TreeView : EStructViewerDisplayMode::ListView;
 	Options.bAllowViewOptions = !bHideViewOptions;
 
-	FOnStructPicked OnPicked(FOnStructPicked::CreateRaw(this, &FInstancedStructDetails::OnStructPicked));
+	FOnStructPicked OnPicked(FOnStructPicked::CreateSP(this, &FInstancedStructDetails::OnStructPicked));
 
 	return SNew(SBox)
 		.WidthOverride(280)
