PropertyBag: Fixed container type hash
- Container type hash was hashing a pointer instead of contents of the array

diff --git a/Engine/Plugins/Experimental/StructUtils/Source/StructUtils/Public/PropertyBag.h b/Engine/Plugins/Experimental/StructUtils/Source/StructUtils/Public/PropertyBag.h
--- a/Engine/Plugins/Experimental/StructUtils/Source/StructUtils/Public/PropertyBag.h
+++ b/Engine/Plugins/Experimental/StructUtils/Source/StructUtils/Public/PropertyBag.h
@@ -5,6 +5,7 @@
 #include "Misc/TVariantMeta.h"
 #include "StructView.h"
 #include "Templates/ValueOrError.h"
+#include "Containers/StaticArray.h"
 #include "PropertyBag.generated.h"
 
 /** Property bag property type, loosely based on BluePrint pin types. */
@@ -47,14 +48,9 @@ struct STRUCTUTILS_API FPropertyBagContainerTypes
 {
 	GENERATED_BODY()
 
-	FPropertyBagContainerTypes()
-		: Types{ EPropertyBagContainerType::None, EPropertyBagContainerType::None }
-		, NumContainers(0)
-	{}
+	FPropertyBagContainerTypes() = default;
 
 	explicit FPropertyBagContainerTypes(EPropertyBagContainerType ContainerType)
-		: Types{ EPropertyBagContainerType::None, EPropertyBagContainerType::None }
-		, NumContainers(0)
 	{
 		if (ContainerType != EPropertyBagContainerType::None)
 		{
@@ -63,10 +59,8 @@ struct STRUCTUTILS_API FPropertyBagContainerTypes
 	}
 
 	FPropertyBagContainerTypes(const std::initializer_list<EPropertyBagContainerType>& InTypes)
-		: Types{ EPropertyBagContainerType::None, EPropertyBagContainerType::None }
-		, NumContainers(0)
 	{
-		for (EPropertyBagContainerType ContainerType : InTypes)
+		for (const EPropertyBagContainerType ContainerType : InTypes)
 		{
 			if (ContainerType != EPropertyBagContainerType::None)
 			{
@@ -75,7 +69,7 @@ struct STRUCTUTILS_API FPropertyBagContainerTypes
 		}
 	}
 
-	bool Add(EPropertyBagContainerType PropertyBagContainerType)
+	bool Add(const EPropertyBagContainerType PropertyBagContainerType)
 	{
 		if (ensure(NumContainers < MaxNestedTypes))
 		{
@@ -93,7 +87,7 @@ struct STRUCTUTILS_API FPropertyBagContainerTypes
 
 	void Reset()
 	{
-		for (auto& Type : Types)
+		for (EPropertyBagContainerType& Type : Types)
 		{
 			Type = EPropertyBagContainerType::None;
 		}
@@ -144,7 +138,12 @@ struct STRUCTUTILS_API FPropertyBagContainerTypes
 
 	friend FORCEINLINE uint32 GetTypeHash(const FPropertyBagContainerTypes& PropertyBagContainerTypes)
 	{
-		return HashCombine(GetTypeHash(PropertyBagContainerTypes.Types), GetTypeHash(PropertyBagContainerTypes.NumContainers));
+		uint32 Hash = GetTypeHash(PropertyBagContainerTypes.NumContainers);
+		for (const EPropertyBagContainerType Type : PropertyBagContainerTypes.Types)
+		{
+			Hash = HashCombine(Hash, GetTypeHash(Type));
+		}
+		return Hash;
 	}
 
 	EPropertyBagContainerType* begin() { return &Types[0]; }
@@ -155,7 +154,7 @@ struct STRUCTUTILS_API FPropertyBagContainerTypes
 protected:
 	static constexpr uint8 MaxNestedTypes = 2;
 
-	EPropertyBagContainerType Types[MaxNestedTypes];
+	TStaticArray<EPropertyBagContainerType, MaxNestedTypes> Types = TStaticArray<EPropertyBagContainerType, MaxNestedTypes>(InPlace, EPropertyBagContainerType::None);
 	uint8 NumContainers = 0;
 };
 
diff --git a/Engine/Plugins/Experimental/StructUtils/Source/StructUtilsTestSuite/Private/PropertyBagTest.cpp b/Engine/Plugins/Experimental/StructUtils/Source/StructUtilsTestSuite/Private/PropertyBagTest.cpp
--- a/Engine/Plugins/Experimental/StructUtils/Source/StructUtilsTestSuite/Private/PropertyBagTest.cpp
+++ b/Engine/Plugins/Experimental/StructUtils/Source/StructUtilsTestSuite/Private/PropertyBagTest.cpp
@@ -469,6 +469,33 @@ struct FTest_Arrays : FAITestBase
 };
 IMPLEMENT_AI_INSTANT_TEST(FTest_Arrays, "System.StructUtils.PropertyBag.Arrays");
 
+struct FTest_SameBag : FAITestBase
+{
+	virtual bool InstantTest() override
+	{
+		static const FName TemperatureName(TEXT("Temperature"));
+		static const FName CountName(TEXT("Count"));
+
+		FInstancedPropertyBag BagA;
+		BagA.AddProperties({
+			{ TemperatureName, EPropertyBagPropertyType::Float },
+			{ CountName, EPropertyBagPropertyType::Int32 }
+		});
+
+
+		FInstancedPropertyBag BagB;
+		BagB.AddProperties({
+			{ TemperatureName, EPropertyBagPropertyType::Float },
+			{ CountName, EPropertyBagPropertyType::Int32 }
+		});
+
+		// Same descriptors should result in same bag struct
+		AITEST_TRUE(TEXT("Property bags should match"), BagA.GetPropertyBagStruct() == BagB.GetPropertyBagStruct());
+
+		return true;
+	}
+};
+IMPLEMENT_AI_INSTANT_TEST(FTest_SameBag, "System.StructUtils.PropertyBag.SameBag");
 
 } // FPropertyBagTest
 
