Added TInstancedStruct to provide a type-safe wrapper for using FInstancedStruct in C++

Note: In the reflection layer TInstancedStruct still produces an FInstanceStruct property, and UHT has been updated to handle that logic

diff --git a/Engine/Plugins/Experimental/StructUtils/Source/StructUtils/Public/InstancedStruct.h b/Engine/Plugins/Experimental/StructUtils/Source/StructUtils/Public/InstancedStruct.h
--- a/Engine/Plugins/Experimental/StructUtils/Source/StructUtils/Public/InstancedStruct.h
+++ b/Engine/Plugins/Experimental/StructUtils/Source/StructUtils/Public/InstancedStruct.h
@@ -8,8 +8,10 @@
 
 struct FConstStructView;
 struct FConstSharedStruct;
+
 /**
  * FInstancedStruct works similarly as instanced UObject* property but is USTRUCTs.
+ * 
  * Example:
  *
  *	UPROPERTY(EditAnywhere, Category = Foo, meta = (BaseStruct = "/Script/ModuleName.TestStructBase"))
@@ -265,6 +267,185 @@ struct TStructOpsTypeTraits<FInstancedStruct> : public TStructOpsTypeTraitsBase2
 	};
 };
 
+/**
+ * TInstancedStruct is a type-safe FInstancedStruct wrapper against the given BaseStruct type.
+ * @note When used as a property, this automatically defines the BaseStruct property meta-data.
+ * 
+ * Example:
+ *
+ *	UPROPERTY(EditAnywhere, Category = Foo)
+ *	TInstancedStruct<FTestStructBase> Test;
+ *
+ *	UPROPERTY(EditAnywhere, Category = Foo)
+ *	TArray<TInstancedStruct<FTestStructBase>> TestArray;
+ */
+template<typename BaseStructT>
+struct TInstancedStruct
+{
+public:
+	TInstancedStruct() = default;
+
+	template<typename T = BaseStructT, typename = std::enable_if_t<std::is_base_of_v<BaseStructT, std::decay_t<T>>>>
+	TInstancedStruct(const TInstancedStruct<T>& InOther)
+		: InstancedStruct(InOther.InstancedStruct)
+	{
+	}
+
+	template<typename T = BaseStructT, typename = std::enable_if_t<std::is_base_of_v<BaseStructT, std::decay_t<T>>>>
+	TInstancedStruct(TInstancedStruct<T>&& InOther)
+		: InstancedStruct(MoveTemp(InOther.InstancedStruct))
+	{
+	}
+
+	template<typename T = BaseStructT, typename = std::enable_if_t<std::is_base_of_v<BaseStructT, std::decay_t<T>>>>
+	TInstancedStruct& operator=(const TInstancedStruct<T>& InOther)
+	{
+		if (this != &InOther)
+		{
+			InstancedStruct = InOther.InstancedStruct;
+		}
+		return *this;
+	}
+
+	template<typename T = BaseStructT, typename = std::enable_if_t<std::is_base_of_v<BaseStructT, std::decay_t<T>>>>
+	TInstancedStruct& operator=(TInstancedStruct<T>&& InOther)
+	{
+		if (this != &InOther)
+		{
+			InstancedStruct = MoveTemp(InOther.InstancedStruct);
+		}
+		return *this;
+	}
+
+	/** Initializes from a raw struct type and optional data. */
+	void InitializeAsScriptStruct(const UScriptStruct* InScriptStruct, const uint8* InStructMemory = nullptr)
+	{
+		checkf(InScriptStruct->IsChildOf(TBaseStructure<BaseStructT>::Get()), TEXT("ScriptStruct must be a child of BaseStruct!"));
+		InstancedStruct.InitializeAs(InScriptStruct, InStructMemory);
+	}
+
+	/** Initializes from struct type and emplace construct. */
+	template<typename T = BaseStructT, typename... TArgs, typename = std::enable_if_t<std::is_base_of_v<BaseStructT, std::decay_t<T>>>>
+	void InitializeAs(TArgs&&... InArgs)
+	{
+		InstancedStruct.InitializeAs<T>(Forward<TArgs>(InArgs)...);
+	}
+
+	/** Creates a new TInstancedStruct from templated struct type. */
+	template<typename T = BaseStructT, typename = std::enable_if_t<std::is_base_of_v<BaseStructT, std::decay_t<T>>>>
+	static TInstancedStruct Make()
+	{
+		TInstancedStruct This;
+		This.InstancedStruct.InitializeAs(TBaseStructure<T>::Get(), nullptr);
+		return This;
+	}
+
+	/** Creates a new TInstancedStruct from templated struct. */
+	template<typename T = BaseStructT, typename = std::enable_if_t<std::is_base_of_v<BaseStructT, std::decay_t<T>>>>
+	static TInstancedStruct Make(const T& Struct)
+	{
+		TInstancedStruct This;
+		This.InstancedStruct.InitializeAs(TBaseStructure<T>::Get(), reinterpret_cast<const uint8*>(&Struct));
+		return This;
+	}
+
+	/** Creates a new TInstancedStruct from the templated type and forward all arguments to constructor. */
+	template<typename T = BaseStructT, typename... TArgs, typename = std::enable_if_t<std::is_base_of_v<BaseStructT, std::decay_t<T>>>>
+	static TInstancedStruct Make(TArgs&&... InArgs)
+	{
+		TInstancedStruct This;
+		This.InstancedStruct.template InitializeAs<T>(Forward<TArgs>(InArgs)...);
+		return This;
+	}
+
+	/** Returns struct type. */
+	const UScriptStruct* GetScriptStruct() const
+	{
+		return InstancedStruct.GetScriptStruct();
+	}
+
+	/** Returns const pointer to raw struct memory. */
+	const uint8* GetMemory() const
+	{
+		return InstancedStruct.GetMemory();
+	}
+
+	/** Reset to empty. */
+	void Reset()
+	{
+		InstancedStruct.Reset();
+	}
+
+	/** Returns const reference to the struct, this getter assumes that all data is valid. */
+	template<typename T = BaseStructT, typename = std::enable_if_t<std::is_base_of_v<BaseStructT, std::decay_t<T>>>>
+	const T& Get() const
+	{
+		return InstancedStruct.Get<T>();
+	}
+
+	/** Returns const pointer to the struct, or nullptr if cast is not valid. */
+	template<typename T = BaseStructT, typename = std::enable_if_t<std::is_base_of_v<BaseStructT, std::decay_t<T>>>>
+	const T* GetPtr() const
+	{
+		return InstancedStruct.GetPtr<T>();
+	}
+
+	/** Returns a mutable pointer to raw struct memory. */
+	uint8* GetMutableMemory()
+	{
+		return InstancedStruct.GetMutableMemory();
+	}
+
+	/** Returns mutable reference to the struct, this getter assumes that all data is valid. */
+	template<typename T = BaseStructT, typename = std::enable_if_t<std::is_base_of_v<BaseStructT, std::decay_t<T>>>>
+	T& GetMutable()
+	{
+		return InstancedStruct.GetMutable<T>();
+	}
+
+	/** Returns mutable pointer to the struct, or nullptr if cast is not valid. */
+	template<typename T = BaseStructT, typename = std::enable_if_t<std::is_base_of_v<BaseStructT, std::decay_t<T>>>>
+	T* GetMutablePtr()
+	{
+		return InstancedStruct.GetMutablePtr<T>();
+	}
+
+	/** Returns True if the struct is valid.*/
+	bool IsValid() const
+	{
+		return InstancedStruct.IsValid();
+	}
+
+	/** Comparison operators. Deep compares the struct instance when identical. */
+	template<typename T = BaseStructT, typename = std::enable_if_t<std::is_base_of_v<BaseStructT, std::decay_t<T>>>>
+	bool operator==(const TInstancedStruct<T>& Other) const
+	{
+		return InstancedStruct == Other.InstancedStruct;
+	}
+
+	template<typename T = BaseStructT, typename = std::enable_if_t<std::is_base_of_v<BaseStructT, std::decay_t<T>>>>
+	bool operator!=(const TInstancedStruct<T>& Other) const
+	{
+		return InstancedStruct != Other.InstancedStruct;
+	}
+
+	void AddReferencedObjects(class FReferenceCollector& Collector)
+	{
+		InstancedStruct.AddStructReferencedObjects(Collector);
+	}
+
+private:
+	/**
+	 * Note:
+	 *   TInstancedStruct is a wrapper for a FInstancedStruct (rather than inheriting) so that it can provide a locked-down type-safe 
+	 *   API for use in C++, without being able to accidentally take a reference to the untyped API to workaround the restrictions.
+	 * 
+	 *   TInstancedStruct MUST be the same size as FInstancedStruct, as the reflection layer treats a TInstancedStruct as a FInstancedStruct.
+	 *   This means that any reflected APIs (like ExportText) that accept an FInstancedStruct pointer can also accept a TInstancedStruct pointer.
+	 */
+	FInstancedStruct InstancedStruct;
+};
+
 #if WITH_EDITORONLY_DATA
 namespace UE::StructUtils::Private
 {
diff --git a/Engine/Source/Programs/Shared/EpicGames.UHT/Types/Properties/UhtStructProperty.cs b/Engine/Source/Programs/Shared/EpicGames.UHT/Types/Properties/UhtStructProperty.cs
--- a/Engine/Source/Programs/Shared/EpicGames.UHT/Types/Properties/UhtStructProperty.cs
+++ b/Engine/Source/Programs/Shared/EpicGames.UHT/Types/Properties/UhtStructProperty.cs
@@ -457,5 +457,66 @@ namespace EpicGames.UHT.Types
 			return true;
 		}
 		#endregion
+
+		#region Parsing support methods
+		/// <summary>
+		/// Parse a template type
+		/// </summary>
+		/// <param name="propertySettings">Property settings</param>
+		/// <param name="tokenReader">Token reader</param>
+		/// <param name="matchedToken">Token matched for type</param>
+		/// <returns>Referenced class</returns>
+		public static UhtScriptStruct? ParseTemplateScriptStruct(UhtPropertySettings propertySettings, IUhtTokenReader tokenReader, UhtToken matchedToken)
+		{
+			UhtSession session = propertySettings.Outer.Session;
+			UhtToken identifier = new();
+
+			if (tokenReader.TryOptional("const"))
+			{
+				propertySettings.MetaData.Add(UhtNames.NativeConst, "");
+			}
+
+			tokenReader.Optional("struct");
+
+			if (!tokenReader.SkipExpectedType(matchedToken.Value, propertySettings.PropertyCategory == UhtPropertyCategory.Member))
+			{
+				return null;
+			}
+
+			tokenReader
+				.Require('<')
+				.Optional("struct")
+				.RequireIdentifier((ref UhtToken token) => { identifier = token; })
+				.Require('>');
+
+			session.Config!.RedirectTypeIdentifier(ref identifier);
+			return propertySettings.Outer.FindType(UhtFindOptions.SourceName | UhtFindOptions.ScriptStruct, ref identifier, tokenReader) as UhtScriptStruct;
+		}
+		#endregion
+
+		#region Keyword
+		[UhtPropertyType(Keyword = "TInstancedStruct")]
+		[SuppressMessage("CodeQuality", "IDE0051:Remove unused private members", Justification = "Attribute accessed method")]
+		[SuppressMessage("Style", "IDE0060:Remove unused parameter", Justification = "Attribute accessed method")]
+		private static UhtProperty? InstancedStructProperty(UhtPropertyResolvePhase resolvePhase, UhtPropertySettings propertySettings, IUhtTokenReader tokenReader, UhtToken matchedToken)
+		{
+			UhtScriptStruct? baseScriptStruct = UhtStructProperty.ParseTemplateScriptStruct(propertySettings, tokenReader, matchedToken);
+			if (baseScriptStruct == null || baseScriptStruct.Session.FInstancedStruct == null)
+			{
+				return null;
+			}
+
+			if (propertySettings.MetaData.ContainsKey("BaseStruct"))
+			{
+				tokenReader.LogError("BaseStruct metadata is implicitly set from the TInstancedStruct template argument and should not be explicitly specified.");
+				return null;
+			}
+
+			// With TInstancedStruct, BaseStruct is used as a type limiter.
+			UhtStructProperty instancedStructProperty = new UhtStructProperty(propertySettings, baseScriptStruct.Session.FInstancedStruct);
+			instancedStructProperty.MetaData.Add("BaseStruct", baseScriptStruct.PathName);
+			return instancedStructProperty;
+		}
+		#endregion
 	}
 }
diff --git a/Engine/Source/Programs/Shared/EpicGames.UHT/Utils/UhtSession.cs b/Engine/Source/Programs/Shared/EpicGames.UHT/Utils/UhtSession.cs
--- a/Engine/Source/Programs/Shared/EpicGames.UHT/Utils/UhtSession.cs
+++ b/Engine/Source/Programs/Shared/EpicGames.UHT/Utils/UhtSession.cs
@@ -770,6 +770,12 @@ namespace EpicGames.UHT.Utils
 		/// </summary>
 		public UhtClass? INotifyFieldValueChanged { get; set; } = null;
 
+		/// <summary>
+		/// After headers are parsed, returns the FInstancedStruct script struct.  Unlike such properties as 
+		/// "UObject", there is no requirement for FInstancedStruct to be defined.  May be null.
+		/// </summary>
+		public UhtScriptStruct? FInstancedStruct { get; set; } = null;
+
 		private readonly List<UhtPackage> _packages = new();
 		private readonly List<UhtHeaderFile> _headerFiles = new();
 		private readonly List<UhtHeaderFile> _sortedHeaderFiles = new();
@@ -2021,6 +2027,7 @@ namespace EpicGames.UHT.Utils
 				_iinterface = (UhtClass?)FindType(null, UhtFindOptions.SourceName | UhtFindOptions.Class, "IInterface");
 				AActor = (UhtClass?)FindType(null, UhtFindOptions.SourceName | UhtFindOptions.Class, "AActor");
 				INotifyFieldValueChanged = (UhtClass?)FindType(null, UhtFindOptions.SourceName | UhtFindOptions.Class, "INotifyFieldValueChanged");
+				FInstancedStruct = (UhtScriptStruct?)FindType(null, UhtFindOptions.SourceName | UhtFindOptions.ScriptStruct, "FInstancedStruct");
 			});
 		}
 
