[StructUtils] Support for changing default type in PropertyBag detail view

*  Also move Metadata names to a seperate header to reference all available metadata values for StructUtils

diff --git a/Engine/Plugins/Experimental/StructUtils/Source/StructUtils/Private/StructUtilsMetadata.cpp b/Engine/Plugins/Experimental/StructUtils/Source/StructUtils/Private/StructUtilsMetadata.cpp
new file mode 100644
--- /dev/null
+++ b/Engine/Plugins/Experimental/StructUtils/Source/StructUtils/Private/StructUtilsMetadata.cpp
@@ -0,0 +1,9 @@
+// Copyright Epic Games, Inc. All Rights Reserved.
+
+#include "StructUtilsMetadata.h"
+
+namespace UE::StructUtils::Metadata
+{
+	const FName FixedLayoutName = TEXT("FixedLayout");
+	const FName DefaultTypeName = TEXT("DefaultType");
+}
\ No newline at end of file
diff --git a/Engine/Plugins/Experimental/StructUtils/Source/StructUtils/Public/StructUtilsMetadata.h b/Engine/Plugins/Experimental/StructUtils/Source/StructUtils/Public/StructUtilsMetadata.h
new file mode 100644
--- /dev/null
+++ b/Engine/Plugins/Experimental/StructUtils/Source/StructUtils/Public/StructUtilsMetadata.h
@@ -0,0 +1,17 @@
+// Copyright Epic Games, Inc. All Rights Reserved.
+
+#pragma once
+
+// Metadata used by StructUtils
+namespace UE::StructUtils::Metadata
+{
+	// Metadata usable in UPROPERTY for customizing the behavior when displaying the property in a property panel or graph node
+
+	/// FInstancedPropertyBag
+
+	/// [PropertyMetadata] FixedLayout: Indicates that the instanced property bag has a fixed layout and it is not possible to add/remove properties.
+	extern STRUCTUTILS_API const FName FixedLayoutName;
+
+	/// [PropertyMetadata] DefaultType: Default property type when adding a new Property. Should be taken from enum EPropertyBagPropertyType
+	extern STRUCTUTILS_API const FName DefaultTypeName;
+}
\ No newline at end of file
diff --git a/Engine/Plugins/Experimental/StructUtils/Source/StructUtilsEditor/Private/Customizations/PropertyBagDetails.cpp b/Engine/Plugins/Experimental/StructUtils/Source/StructUtilsEditor/Private/Customizations/PropertyBagDetails.cpp
--- a/Engine/Plugins/Experimental/StructUtils/Source/StructUtilsEditor/Private/Customizations/PropertyBagDetails.cpp
+++ b/Engine/Plugins/Experimental/StructUtils/Source/StructUtilsEditor/Private/Customizations/PropertyBagDetails.cpp
@@ -14,7 +14,7 @@
 #include "ScopedTransaction.h"
 #include "Engine/UserDefinedStruct.h"
 #include "SPinTypeSelector.h"
-#include "PropertyBag.h"
+#include "StructUtilsMetadata.h"
 
 #include UE_INLINE_GENERATED_CPP_BY_NAME(PropertyBagDetails)
 
@@ -783,24 +783,35 @@ void FPropertyBagDetails::CustomizeHeader(TSharedRef<IPropertyHandle> StructProp
 	StructProperty = StructPropertyHandle;
 	check(StructProperty);
 	
-	static const FName NAME_FixedLayout = "FixedLayout";
 	static const FName NAME_ShowOnlyInnerProperties = "ShowOnlyInnerProperties";
 	
 	if (const FProperty* MetaDataProperty = StructProperty->GetMetaDataProperty())
 	{
-		bFixedLayout = MetaDataProperty->HasMetaData(NAME_FixedLayout);
+		bFixedLayout = MetaDataProperty->HasMetaData(UE::StructUtils::Metadata::FixedLayoutName);
 
 		// Don't show the header if ShowOnlyInnerProperties is set
 		if (MetaDataProperty->HasMetaData(NAME_ShowOnlyInnerProperties))
 		{
 			return;
 		}
+
+		if (MetaDataProperty->HasMetaData(UE::StructUtils::Metadata::DefaultTypeName))
+		{
+			if (UEnum* Enum = StaticEnum<EPropertyBagPropertyType>())
+			{
+				int32 EnumIndex = Enum->GetIndexByNameString(MetaDataProperty->GetMetaData(UE::StructUtils::Metadata::DefaultTypeName));
+				if (EnumIndex != INDEX_NONE)
+				{
+					DefaultType = EPropertyBagPropertyType(Enum->GetValueByIndex(EnumIndex));
+				}
+			}
+		}
 	}
 
 	TSharedPtr<SWidget> ValueWidget = SNullWidget::NullWidget;
 	if (!bFixedLayout)
 	{
-		ValueWidget = MakeAddPropertyWidget(StructProperty, PropUtils);
+		ValueWidget = MakeAddPropertyWidget(StructProperty, PropUtils, DefaultType);
 	}
 	
 	HeaderRow
@@ -823,7 +834,7 @@ void FPropertyBagDetails::CustomizeChildren(TSharedRef<IPropertyHandle> StructPr
 	StructBuilder.AddCustomBuilder(InstanceDetails);
 }
 
-TSharedPtr<SWidget> FPropertyBagDetails::MakeAddPropertyWidget(TSharedPtr<IPropertyHandle> InStructProperty, TSharedPtr<IPropertyUtilities> InPropUtils)
+TSharedPtr<SWidget> FPropertyBagDetails::MakeAddPropertyWidget(TSharedPtr<IPropertyHandle> InStructProperty, TSharedPtr<IPropertyUtilities> InPropUtils, EPropertyBagPropertyType DefaultType)
 {
 	return SNew(SHorizontalBox)
 		+SHorizontalBox::Slot()
@@ -834,7 +845,7 @@ TSharedPtr<SWidget> FPropertyBagDetails::MakeAddPropertyWidget(TSharedPtr<IPrope
 			.HAlign(HAlign_Center)
 			.ButtonStyle(FAppStyle::Get(), "SimpleButton")
 			.ToolTipText(LOCTEXT("AddProperty_Tooltip", "Add new property"))
-			.OnClicked_Lambda([StructProperty = InStructProperty, PropUtils = InPropUtils]()
+			.OnClicked_Lambda([StructProperty = InStructProperty, PropUtils = InPropUtils, DefaultType]()
 			{
 				constexpr int32 MaxIterations = 100;
 				FName NewName(TEXT("NewProperty"));
@@ -851,9 +862,9 @@ TSharedPtr<SWidget> FPropertyBagDetails::MakeAddPropertyWidget(TSharedPtr<IPrope
 
 				UE::StructUtils::Private::ApplyChangesToPropertyDescs(
 					LOCTEXT("OnPropertyAdded", "Add Property"), StructProperty, PropUtils,
-					[&NewName](TArray<FPropertyBagPropertyDesc>& PropertyDescs)
+					[&NewName, DefaultType](TArray<FPropertyBagPropertyDesc>& PropertyDescs)
 					{
-						PropertyDescs.Emplace(NewName, EPropertyBagPropertyType::Bool);
+						PropertyDescs.Emplace(NewName, DefaultType);
 					});
 					
 				return FReply::Handled();
diff --git a/Engine/Plugins/Experimental/StructUtils/Source/StructUtilsEditor/Public/PropertyBagDetails.h b/Engine/Plugins/Experimental/StructUtils/Source/StructUtilsEditor/Public/PropertyBagDetails.h
--- a/Engine/Plugins/Experimental/StructUtils/Source/StructUtilsEditor/Public/PropertyBagDetails.h
+++ b/Engine/Plugins/Experimental/StructUtils/Source/StructUtilsEditor/Public/PropertyBagDetails.h
@@ -4,6 +4,7 @@
 
 #include "InstancedStructDetails.h"
 #include "EdGraphSchema_K2.h"
+#include "PropertyBag.h"
 #include "PropertyBagDetails.generated.h"
 
 class FReply;
@@ -26,7 +27,7 @@ public:
 	static TSharedRef<IPropertyTypeCustomization> MakeInstance();
 
 	/** Creates add property widget. */
-	static TSharedPtr<SWidget> MakeAddPropertyWidget(TSharedPtr<IPropertyHandle> InStructProperty, TSharedPtr<IPropertyUtilities> InPropUtils);
+	static TSharedPtr<SWidget> MakeAddPropertyWidget(TSharedPtr<IPropertyHandle> InStructProperty, TSharedPtr<IPropertyUtilities> InPropUtils, EPropertyBagPropertyType DefaultType = EPropertyBagPropertyType::Bool);
 	
 protected:
 	/** IPropertyTypeCustomization interface */
@@ -40,7 +41,8 @@ private:
 	/** Handle to the struct property being edited */
 	TSharedPtr<IPropertyHandle> StructProperty;
 	TSharedPtr<IPropertyUtilities> PropUtils;
-	bool bFixedLayout = false; 
+	bool bFixedLayout = false;
+	EPropertyBagPropertyType DefaultType = EPropertyBagPropertyType::Bool;
 };
 
 /** 
