[StructUtils] Tentative fix Non-Unity build

#rnx

diff --git a/Engine/Plugins/Experimental/StructUtils/Source/StructUtils/Public/StructUtilsMetadata.h b/Engine/Plugins/Experimental/StructUtils/Source/StructUtils/Public/StructUtilsMetadata.h
--- a/Engine/Plugins/Experimental/StructUtils/Source/StructUtils/Public/StructUtilsMetadata.h
+++ b/Engine/Plugins/Experimental/StructUtils/Source/StructUtils/Public/StructUtilsMetadata.h
@@ -2,6 +2,8 @@
 
 #pragma once
 
+#include "UObject/NameTypes.h"
+
 // Metadata used by StructUtils
 namespace UE::StructUtils::Metadata
 {
