[StructUtils] Add filtering on containers you can use with Properties

diff --git a/Engine/Plugins/Experimental/StructUtils/Source/StructUtils/Private/StructUtilsMetadata.cpp b/Engine/Plugins/Experimental/StructUtils/Source/StructUtils/Private/StructUtilsMetadata.cpp
--- a/Engine/Plugins/Experimental/StructUtils/Source/StructUtils/Private/StructUtilsMetadata.cpp
+++ b/Engine/Plugins/Experimental/StructUtils/Source/StructUtils/Private/StructUtilsMetadata.cpp
@@ -6,4 +6,5 @@ namespace UE::StructUtils::Metadata
 {
 	const FName FixedLayoutName = TEXT("FixedLayout");
 	const FName DefaultTypeName = TEXT("DefaultType");
+	const FName AllowArraysName = TEXT("AllowArrays");
 }
\ No newline at end of file
diff --git a/Engine/Plugins/Experimental/StructUtils/Source/StructUtils/Public/StructUtilsMetadata.h b/Engine/Plugins/Experimental/StructUtils/Source/StructUtils/Public/StructUtilsMetadata.h
--- a/Engine/Plugins/Experimental/StructUtils/Source/StructUtils/Public/StructUtilsMetadata.h
+++ b/Engine/Plugins/Experimental/StructUtils/Source/StructUtils/Public/StructUtilsMetadata.h
@@ -16,4 +16,7 @@ namespace UE::StructUtils::Metadata
 
 	/// [PropertyMetadata] DefaultType: Default property type when adding a new Property. Should be taken from enum EPropertyBagPropertyType
 	extern STRUCTUTILS_API const FName DefaultTypeName;
+
+	/// [PropertyMetadata] AllowArrays: By default it is always allowed to create arrays for properties. Use AllowArrays=false to disable Array support.
+	extern STRUCTUTILS_API const FName AllowArraysName;
 }
\ No newline at end of file
diff --git a/Engine/Plugins/Experimental/StructUtils/Source/StructUtilsEditor/Private/Customizations/PropertyBagDetails.cpp b/Engine/Plugins/Experimental/StructUtils/Source/StructUtilsEditor/Private/Customizations/PropertyBagDetails.cpp
--- a/Engine/Plugins/Experimental/StructUtils/Source/StructUtilsEditor/Private/Customizations/PropertyBagDetails.cpp
+++ b/Engine/Plugins/Experimental/StructUtils/Source/StructUtilsEditor/Private/Customizations/PropertyBagDetails.cpp
@@ -470,11 +470,12 @@ bool CanHaveMemberVariableOfType(const FEdGraphPinType& PinType)
 //  - ChildPropertyHandle a child property of the FInstancedPropertyBag::Value (FInstancedStruct)  
 //----------------------------------------------------------------//
 
-FPropertyBagInstanceDataDetails::FPropertyBagInstanceDataDetails(TSharedPtr<IPropertyHandle> InStructProperty, const TSharedPtr<IPropertyUtilities>& InPropUtils, const bool bInFixedLayout)
+FPropertyBagInstanceDataDetails::FPropertyBagInstanceDataDetails(TSharedPtr<IPropertyHandle> InStructProperty, const TSharedPtr<IPropertyUtilities>& InPropUtils, const bool bInFixedLayout, const bool bInAllowArrays)
 	: FInstancedStructDataDetails(InStructProperty.IsValid() ? InStructProperty->GetChildHandle(TEXT("Value")) : nullptr)
 	, BagStructProperty(InStructProperty)
 	, PropUtils(InPropUtils)
 	, bFixedLayout(bInFixedLayout)
+	, bAllowArrays(bInAllowArrays)
 {
 	ensure(UE::StructUtils::Private::IsScriptStruct<FInstancedPropertyBag>(BagStructProperty));
 	ensure(PropUtils != nullptr);
@@ -727,7 +728,7 @@ TSharedRef<SWidget> FPropertyBagInstanceDataDetails::OnPropertyNameContent(TShar
 				.TargetPinType_Lambda(GetPinInfo)
 				.OnPinTypeChanged_Lambda(PinInfoChanged)
 				.Schema(GetDefault<UPropertyBagSchema>())
-				.bAllowArrays(true)
+				.bAllowArrays(bAllowArrays)
 				.TypeTreeFilter(ETypeTreeFilter::None)
 				.Font( IDetailLayoutBuilder::GetDetailFont() )
 		],
@@ -789,6 +790,10 @@ void FPropertyBagDetails::CustomizeHeader(TSharedRef<IPropertyHandle> StructProp
 	{
 		bFixedLayout = MetaDataProperty->HasMetaData(UE::StructUtils::Metadata::FixedLayoutName);
 
+		bAllowArrays = MetaDataProperty->HasMetaData(UE::StructUtils::Metadata::AllowArraysName)
+			? MetaDataProperty->GetBoolMetaData(UE::StructUtils::Metadata::AllowArraysName)
+			: true;
+
 		// Don't show the header if ShowOnlyInnerProperties is set
 		if (MetaDataProperty->HasMetaData(NAME_ShowOnlyInnerProperties))
 		{
@@ -830,7 +835,7 @@ void FPropertyBagDetails::CustomizeHeader(TSharedRef<IPropertyHandle> StructProp
 void FPropertyBagDetails::CustomizeChildren(TSharedRef<IPropertyHandle> StructPropertyHandle, IDetailChildrenBuilder& StructBuilder, IPropertyTypeCustomizationUtils& StructCustomizationUtils)
 {
 	// Show the Value (FInstancedStruct) as child rows.
-	const TSharedRef<FPropertyBagInstanceDataDetails> InstanceDetails = MakeShareable(new FPropertyBagInstanceDataDetails(StructProperty, PropUtils, bFixedLayout));
+	const TSharedRef<FPropertyBagInstanceDataDetails> InstanceDetails = MakeShareable(new FPropertyBagInstanceDataDetails(StructProperty, PropUtils, bFixedLayout, bAllowArrays));
 	StructBuilder.AddCustomBuilder(InstanceDetails);
 }
 
diff --git a/Engine/Plugins/Experimental/StructUtils/Source/StructUtilsEditor/Public/PropertyBagDetails.h b/Engine/Plugins/Experimental/StructUtils/Source/StructUtilsEditor/Public/PropertyBagDetails.h
--- a/Engine/Plugins/Experimental/StructUtils/Source/StructUtilsEditor/Public/PropertyBagDetails.h
+++ b/Engine/Plugins/Experimental/StructUtils/Source/StructUtilsEditor/Public/PropertyBagDetails.h
@@ -41,8 +41,9 @@ private:
 	/** Handle to the struct property being edited */
 	TSharedPtr<IPropertyHandle> StructProperty;
 	TSharedPtr<IPropertyUtilities> PropUtils;
-	bool bFixedLayout = false;
 	EPropertyBagPropertyType DefaultType = EPropertyBagPropertyType::Bool;
+	bool bFixedLayout = false;
+	bool bAllowArrays = true;
 };
 
 /** 
@@ -56,7 +57,7 @@ private:
 class STRUCTUTILSEDITOR_API FPropertyBagInstanceDataDetails : public FInstancedStructDataDetails
 {
 public:
-	FPropertyBagInstanceDataDetails(TSharedPtr<IPropertyHandle> InStructProperty, const TSharedPtr<IPropertyUtilities>& InPropUtils, const bool bInFixedLayout);
+	FPropertyBagInstanceDataDetails(TSharedPtr<IPropertyHandle> InStructProperty, const TSharedPtr<IPropertyUtilities>& InPropUtils, const bool bInFixedLayout, const bool bInAllowArrays = true);
 
 	virtual void OnChildRowAdded(IDetailPropertyRow& ChildRow) override;
 
@@ -66,6 +67,7 @@ protected:
 	TSharedPtr<IPropertyHandle> BagStructProperty;
 	TSharedPtr<IPropertyUtilities> PropUtils;
 	bool bFixedLayout = false;
+	bool bAllowArrays = true;
 };
 
 /**
