StructUtils: Added User Defined Struct support for instanced struct and property bag
- Enabled UDS in the struct/property pickers
- Added support for reinstancing instanced struct & co when UDS layout changes

diff --git a/Engine/Plugins/Experimental/StructUtils/Source/StructUtils/Private/InstancedStruct.cpp b/Engine/Plugins/Experimental/StructUtils/Source/StructUtils/Private/InstancedStruct.cpp
--- a/Engine/Plugins/Experimental/StructUtils/Source/StructUtils/Private/InstancedStruct.cpp
+++ b/Engine/Plugins/Experimental/StructUtils/Source/StructUtils/Private/InstancedStruct.cpp
@@ -2,11 +2,17 @@
 #include "InstancedStruct.h"
 #include "StructView.h"
 #include "Serialization/PropertyLocalizationDataGathering.h"
+#include "StructUtilsTypes.h"
 
 #if WITH_ENGINE
 #include "Engine/PackageMapClient.h"
 #include "Engine/NetConnection.h"
+#include "Engine/UserDefinedStruct.h"
 #include "Net/RepLayout.h"
+#include "Serialization/MemoryWriter.h"
+#include "Serialization/MemoryReader.h"
+#include "Serialization/ArchiveUObject.h"
+#include "Serialization/ObjectAndNameAsStringProxyArchive.h"
 #endif // WITH_ENGINE
 
 #include UE_INLINE_GENERATED_CPP_BY_NAME(InstancedStruct)
@@ -210,7 +216,22 @@ bool FInstancedStruct::Serialize(FArchive& Ar)
 	else if (Ar.IsSaving())
 	{
 		// UScriptStruct type
-		Ar << NonConstStruct;
+#if WITH_ENGINE && WITH_EDITOR
+		const UUserDefinedStruct* UserDefinedStruct = Cast<UUserDefinedStruct>(NonConstStruct);
+		if (UserDefinedStruct
+			&& UserDefinedStruct->Status == EUserDefinedStructureStatus::UDSS_Duplicate
+			&& UserDefinedStruct->PrimaryStruct.IsValid())
+		{
+			// If saving a duplicated UDS, save the primary type instead, so that the data is loaded with the original struct.
+			// This is used as part of the user defined struct reinstancing logic.
+			UUserDefinedStruct* PrimaryUserDefinedStruct = UserDefinedStruct->PrimaryStruct.Get(); 
+			Ar << PrimaryUserDefinedStruct;
+		}
+		else
+#endif			
+		{
+			Ar << NonConstStruct;
+		}
 	
 		// Size of the serialized memory (reserve location)
 		const int64 SizeOffset = Ar.Tell(); // Position to write the actual size after struct serialization
@@ -409,8 +430,52 @@ bool FInstancedStruct::Identical(const FInstancedStruct* Other, uint32 PortFlags
 	return true;
 }
 
-void FInstancedStruct::AddStructReferencedObjects(class FReferenceCollector& Collector)
+void FInstancedStruct::AddStructReferencedObjects(FReferenceCollector& Collector)
 {
+#if WITH_ENGINE && WITH_EDITOR
+	// Reference collector is used to visit all instances of instanced structs and replace their contents.
+	if (const UUserDefinedStruct* StructureToReinstance = UE::StructUtils::Private::GetStructureToReinstance())
+	{
+		if (const UUserDefinedStruct* UserDefinedStruct = Cast<UUserDefinedStruct>(ScriptStruct))
+		{
+			if (StructureToReinstance->Status == EUserDefinedStructureStatus::UDSS_Duplicate)
+			{
+				// On the first pass we replace the UDS with a duplicate that represents the currently allocated struct.
+				// GStructureToReinstance is the duplicated struct, and StructureToReinstance->PrimaryStruct is the UDS that is being reinstanced.
+				
+				if (UserDefinedStruct == StructureToReinstance->PrimaryStruct)
+				{
+					ScriptStruct = StructureToReinstance;
+				}
+			}
+			else
+			{
+				// On the second pass we reinstantiate the data using serialization.
+				// When saving, the UDSs are written using the duplicate which represents current layout, but PrimaryStruct is serialized as the type.
+				// When reading, the data is initialized with the new type, and the serialization will take care of reading from the old data.
+
+				if (UserDefinedStruct->PrimaryStruct == StructureToReinstance)
+				{
+					if (UObject* Outer = UE::StructUtils::Private::GetCurrentReinstanceOuterObject())
+					{
+						Outer->MarkPackageDirty();
+					}
+
+					TArray<uint8> Data;
+					
+					FMemoryWriter Writer(Data);
+					FObjectAndNameAsStringProxyArchive WriterProxy(Writer, /*bInLoadIfFindFails*/true);
+					Serialize(WriterProxy);
+
+					FMemoryReader Reader(Data);
+					FObjectAndNameAsStringProxyArchive ReaderProxy(Reader, /*bInLoadIfFindFails*/true);
+					Serialize(ReaderProxy);
+				}
+			}
+		}
+	}
+#endif
+	
 	if (ScriptStruct != nullptr)
 	{
 		Collector.AddReferencedObject(ScriptStruct);
@@ -418,6 +483,13 @@ void FInstancedStruct::AddStructReferencedObjects(class FReferenceCollector& Col
 	}
 }
 
+#if WITH_ENGINE && WITH_EDITOR
+void FInstancedStruct::ReplaceScriptStructInternal(const UScriptStruct* NewStruct)
+{
+	ScriptStruct = NewStruct;
+}
+#endif
+
 bool FInstancedStruct::NetSerialize(FArchive& Ar, UPackageMap* Map, bool& bOutSuccess)
 {
 #if WITH_ENGINE
diff --git a/Engine/Plugins/Experimental/StructUtils/Source/StructUtils/Private/InstancedStructContainer.cpp b/Engine/Plugins/Experimental/StructUtils/Source/StructUtils/Private/InstancedStructContainer.cpp
--- a/Engine/Plugins/Experimental/StructUtils/Source/StructUtils/Private/InstancedStructContainer.cpp
+++ b/Engine/Plugins/Experimental/StructUtils/Source/StructUtils/Private/InstancedStructContainer.cpp
@@ -1,6 +1,14 @@
 // Copyright Epic Games, Inc. All Rights Reserved.
 #include "InstancedStructContainer.h"
 
+#if WITH_ENGINE && WITH_EDITOR
+#include "Engine/UserDefinedStruct.h"
+#include "Serialization/MemoryWriter.h"
+#include "Serialization/MemoryReader.h"
+#include "Serialization/ArchiveUObject.h"
+#include "Serialization/ObjectAndNameAsStringProxyArchive.h"
+#endif
+
 #include UE_INLINE_GENERATED_CPP_BY_NAME(InstancedStructContainer)
 
 FInstancedStructContainer::FInstancedStructContainer()
@@ -365,8 +373,68 @@ void FInstancedStructContainer::Empty()
 	AllocatedSize = 0;
 }
 
-void FInstancedStructContainer::AddStructReferencedObjects(FReferenceCollector& Collector) const
+void FInstancedStructContainer::AddStructReferencedObjects(FReferenceCollector& Collector)
 {
+#if WITH_ENGINE && WITH_EDITOR
+	if (const UUserDefinedStruct* StructureToReinstance = UE::StructUtils::Private::GetStructureToReinstance())
+	{
+		bool bContainsUserDefineStruct = false;
+		for (int32 Index = 0, Num = NumItems; Index < Num; Index++)
+		{
+			FItem& Item = GetItem(Index);
+			if (const UUserDefinedStruct* UserDefinedStruct = Cast<UUserDefinedStruct>(Item.ScriptStruct))
+			{
+				if (UserDefinedStruct == StructureToReinstance
+					|| UserDefinedStruct == StructureToReinstance->PrimaryStruct
+					|| UserDefinedStruct->PrimaryStruct == StructureToReinstance
+					|| UserDefinedStruct->PrimaryStruct == StructureToReinstance->PrimaryStruct)
+				{
+					bContainsUserDefineStruct = true;
+					break;
+				}
+			}
+		}
+		
+		if (bContainsUserDefineStruct)
+		{
+			if (StructureToReinstance->Status == EUserDefinedStructureStatus::UDSS_Duplicate)
+			{
+				// On the first pass we replace the UDS with a duplicate that represents the currently allocated struct.
+				// StructureToReinstance is the duplicated struct, and StructureToReinstance->PrimaryStruct is the UDS that is being reinstanced.
+				for (int32 Index = 0, Num = NumItems; Index < Num; Index++)
+				{
+					FItem& Item = GetItem(Index);
+					if (Item.ScriptStruct == StructureToReinstance->PrimaryStruct)
+					{
+						Item.ScriptStruct = StructureToReinstance;
+					}
+				}
+			}
+			else
+			{
+				// On the second pass we reinstantiate the data using serialization.
+				// When saving, the UDSs are written using the duplicate which represents current layout, but PrimaryStruct is serialized as the type.
+				// When reading, the data is initialized with the new type, and the serialization will take care of reading from the old data.
+				
+				if (UObject* Outer = UE::StructUtils::Private::GetCurrentReinstanceOuterObject())
+				{
+					Outer->MarkPackageDirty();
+				}
+				
+				TArray<uint8> Data;
+
+				FMemoryWriter Writer(Data);
+				FObjectAndNameAsStringProxyArchive WriterProxy(Writer, /*bInLoadIfFindFails*/true);
+				Serialize(WriterProxy);
+
+				FMemoryReader Reader(Data);
+				FObjectAndNameAsStringProxyArchive ReaderProxy(Reader, /*bInLoadIfFindFails*/true);
+				Serialize(ReaderProxy);
+			}
+		}
+	}
+#endif
+	
 	for (int32 Index = 0, Num = NumItems; Index < Num; Index++)
 	{
 		FItem& Item = GetItem(Index);
@@ -512,7 +580,23 @@ bool FInstancedStructContainer::Serialize(FArchive& Ar)
 			for (int32 Index = 0; Index < NumItems; Index++)
 			{
 				const FItem& Item = GetItem(Index);
-				Ar << ConstCast(Item.ScriptStruct);
+				
+#if WITH_ENGINE && WITH_EDITOR
+				UUserDefinedStruct* UserDefinedStruct = Cast<UUserDefinedStruct>(ConstCast(Item.ScriptStruct);
+				if (UserDefinedStruct
+					&& UserDefinedStruct->Status == EUserDefinedStructureStatus::UDSS_Duplicate
+					&& UserDefinedStruct->PrimaryStruct.IsValid())
+				{
+					// If saving a duplicated UDS, save the primary type instead, so that the data is loaded with the original struct.
+					// This is used as part of the user defined struct reinstancing logic.
+					UUserDefinedStruct* PrimaryUserDefinedStruct = UserDefinedStruct->PrimaryStruct.Get(); 
+					Ar << PrimaryUserDefinedStruct;
+				}
+				else
+#endif			
+				{
+					Ar << ConstCast(Item.ScriptStruct);
+				}
 			}
 
 			// Save item values
diff --git a/Engine/Plugins/Experimental/StructUtils/Source/StructUtils/Private/PropertyBag.cpp b/Engine/Plugins/Experimental/StructUtils/Source/StructUtils/Private/PropertyBag.cpp
--- a/Engine/Plugins/Experimental/StructUtils/Source/StructUtils/Private/PropertyBag.cpp
+++ b/Engine/Plugins/Experimental/StructUtils/Source/StructUtils/Private/PropertyBag.cpp
@@ -1,10 +1,19 @@
 // Copyright Epic Games, Inc. All Rights Reserved.
 
 #include "PropertyBag.h"
+#include "Engine/UserDefinedStruct.h"
 #include "UObject/EnumProperty.h"
 #include "UObject/Package.h"
 #include "UObject/TextProperty.h"
 
+#if WITH_ENGINE && WITH_EDITOR
+#include "Engine/UserDefinedStruct.h"
+#include "Serialization/MemoryWriter.h"
+#include "Serialization/MemoryReader.h"
+#include "Serialization/ArchiveUObject.h"
+#include "Serialization/ObjectAndNameAsStringProxyArchive.h"
+#endif
+
 #include UE_INLINE_GENERATED_CPP_BY_NAME(PropertyBag)
 
 struct STRUCTUTILS_API FPropertyBagCustomVersion
@@ -1798,7 +1807,27 @@ bool FInstancedPropertyBag::Serialize(FArchive& Ar)
 		else if (Ar.IsSaving())
 		{
 			check(BagStruct);
-			Ar << BagStruct->PropertyDescs;
+			
+			TArray<FPropertyBagPropertyDesc> PropertyDescs = BagStruct->PropertyDescs;
+#if WITH_ENGINE && WITH_EDITOR
+			// Save primary struct for user defined struct properties.
+			// This is used as part of the user defined struct reinstancing logic.
+			for (FPropertyBagPropertyDesc& Desc : PropertyDescs)
+			{
+				if (Desc.ValueType == EPropertyBagPropertyType::Struct)
+				{
+					const UUserDefinedStruct* UserDefinedStruct = Cast<UUserDefinedStruct>(Desc.ValueTypeObject);
+					if (UserDefinedStruct
+						&& UserDefinedStruct->Status == EUserDefinedStructureStatus::UDSS_Duplicate
+						&& UserDefinedStruct->PrimaryStruct.IsValid())
+					{
+						Desc.ValueTypeObject = UserDefinedStruct->PrimaryStruct.Get();
+					}
+				}
+			}
+#endif			
+			
+			Ar << PropertyDescs;
 
 			const int64 SizeOffset = Ar.Tell(); // Position to write the actual size after struct serialization
 			int32 SerialSize = 0;
@@ -1823,6 +1852,64 @@ bool FInstancedPropertyBag::Serialize(FArchive& Ar)
 	return true;
 }
 
+void FInstancedPropertyBag::AddStructReferencedObjects(FReferenceCollector& Collector)
+{
+#if WITH_ENGINE && WITH_EDITOR
+	// Reference collector is used to visit all instances of instanced structs the the like when a user defined struct is reinstanced.
+	if (const UUserDefinedStruct* StructureToReinstance = UE::StructUtils::Private::GetStructureToReinstance())
+	{
+		const UPropertyBag* Bag = GetPropertyBagStruct();
+		if (Bag && Bag->ContainsUserDefinedStruct(StructureToReinstance))
+		{
+			if (StructureToReinstance->Status == EUserDefinedStructureStatus::UDSS_Duplicate)
+			{
+				// On the first pass we create a new bag that contains copy of UDS that represents the currently allocated struct.
+				// GStructureToReinstance is the duplicated struct, and GStructureToReinstance->PrimaryStruct is the UDS that is being reinstanced.
+
+				TArray<FPropertyBagPropertyDesc> PropertyDescs = Bag->PropertyDescs;
+				for (FPropertyBagPropertyDesc& Desc : PropertyDescs)
+				{
+					if (Desc.ValueType == EPropertyBagPropertyType::Struct)
+					{
+						if (UUserDefinedStruct* UserDefinedStruct = Cast<UUserDefinedStruct>(Desc.ValueTypeObject))
+						{
+							if (UserDefinedStruct == StructureToReinstance->PrimaryStruct)
+							{
+								Desc.ValueTypeObject = StructureToReinstance;
+							}
+						}					
+					}
+				}
+
+				const UPropertyBag* NewBag = UPropertyBag::GetOrCreateFromDescs(PropertyDescs);
+				Value.ReplaceScriptStructInternal(NewBag);
+			}
+			else
+			{
+				// On the second pass we reinstantiate the data using serialization.
+				// When saving, the UDSs are written using the duplicate which represents current layout, but PrimaryStruct is serialized as the type.
+				// When reading, the data is initialized with the new type, and the serialization will take care of reading from the old data.
+
+				if (UObject* Outer = UE::StructUtils::Private::GetCurrentReinstanceOuterObject())
+				{
+					Outer->MarkPackageDirty();
+				}
+				
+				TArray<uint8> Data;
+
+				FMemoryWriter Writer(Data);
+				FObjectAndNameAsStringProxyArchive WriterProxy(Writer, /*bInLoadIfFindFails*/true);
+				Serialize(WriterProxy);
+				
+				FMemoryReader Reader(Data);
+				FObjectAndNameAsStringProxyArchive ReaderProxy(Reader, /*bInLoadIfFindFails*/true);
+				Serialize(ReaderProxy);
+			}
+		}
+	}
+#endif	
+}
+
 
 //----------------------------------------------------------------//
 //  FPropertyBagArrayRef
@@ -2180,6 +2267,34 @@ const UPropertyBag* UPropertyBag::GetOrCreateFromDescs(const TConstArrayView<FPr
 	return NewBag;
 }
 
+#if WITH_ENGINE && WITH_EDITOR
+bool UPropertyBag::ContainsUserDefinedStruct(const UUserDefinedStruct* UserDefinedStruct) const
+{
+	if (!UserDefinedStruct)
+	{
+		return false;
+	}
+	
+	for (const FPropertyBagPropertyDesc& Desc : PropertyDescs)
+	{
+		if (Desc.ValueType == EPropertyBagPropertyType::Struct)
+		{
+			if (const UUserDefinedStruct* OwnedUserDefinedStruct = Cast<UUserDefinedStruct>(Desc.ValueTypeObject))
+			{
+				if (OwnedUserDefinedStruct == UserDefinedStruct
+					|| OwnedUserDefinedStruct->PrimaryStruct == UserDefinedStruct
+					|| OwnedUserDefinedStruct == UserDefinedStruct->PrimaryStruct
+					|| OwnedUserDefinedStruct->PrimaryStruct == UserDefinedStruct->PrimaryStruct)
+				{
+					return true;
+				}
+			}
+		}
+	}
+	return false;
+}
+#endif
+
 void UPropertyBag::InitializeStruct(void* Dest, int32 ArrayDim) const
 {
 	Super::InitializeStruct(Dest, ArrayDim);
diff --git a/Engine/Plugins/Experimental/StructUtils/Source/StructUtils/Private/StructUtilsTypes.cpp b/Engine/Plugins/Experimental/StructUtils/Source/StructUtils/Private/StructUtilsTypes.cpp
--- a/Engine/Plugins/Experimental/StructUtils/Source/StructUtils/Private/StructUtilsTypes.cpp
+++ b/Engine/Plugins/Experimental/StructUtils/Source/StructUtils/Private/StructUtilsTypes.cpp
@@ -48,4 +48,45 @@ namespace UE::StructUtils
 	{
 		return GetStructCrc32Helper(SharedView, CRC);
 	}
+
+	namespace Private
+	{
+#if WITH_ENGINE && WITH_EDITOR
+		const UUserDefinedStruct* GStructureToReinstance = nullptr;
+		UObject* GCurrentReinstanceOuterObject = nullptr;
+
+		FStructureToReinstanceScope::FStructureToReinstanceScope(const UUserDefinedStruct* StructureToReinstance)
+		{
+			OldStructureToReinstance = GStructureToReinstance;
+			GStructureToReinstance = StructureToReinstance;
+		}
+	
+		FStructureToReinstanceScope::~FStructureToReinstanceScope()
+		{
+			GStructureToReinstance = OldStructureToReinstance;
+		}
+
+		FCurrentReinstanceOuterObjectScope::FCurrentReinstanceOuterObjectScope(UObject* CurrentReinstanceOuterObject)
+		{
+			OldCurrentReinstanceOuterObject = GCurrentReinstanceOuterObject;
+			GCurrentReinstanceOuterObject = CurrentReinstanceOuterObject;
+		}
+
+		FCurrentReinstanceOuterObjectScope::~FCurrentReinstanceOuterObjectScope()
+		{
+			GCurrentReinstanceOuterObject = OldCurrentReinstanceOuterObject;
+		}
+
+		const UUserDefinedStruct* GetStructureToReinstance()
+		{
+			return GStructureToReinstance;
+		}
+	
+		UObject* GetCurrentReinstanceOuterObject()
+		{
+			return GCurrentReinstanceOuterObject;
+		}
+
+#endif
+	}
 }
diff --git a/Engine/Plugins/Experimental/StructUtils/Source/StructUtils/Public/InstancedStruct.h b/Engine/Plugins/Experimental/StructUtils/Source/StructUtils/Public/InstancedStruct.h
--- a/Engine/Plugins/Experimental/StructUtils/Source/StructUtils/Public/InstancedStruct.h
+++ b/Engine/Plugins/Experimental/StructUtils/Source/StructUtils/Public/InstancedStruct.h
@@ -8,6 +8,7 @@
 
 struct FConstStructView;
 struct FConstSharedStruct;
+class UUserDefinedStruct;
 
 /**
  * FInstancedStruct works similarly as instanced UObject* property but is USTRUCTs.
@@ -231,6 +232,11 @@ public:
 		return !Identical(&Other, PPF_None);
 	}
 
+#if WITH_ENGINE && WITH_EDITOR
+	/** Internal method used to replace the script struct during user defined struct instantiation. */
+	void ReplaceScriptStructInternal(const UScriptStruct* NewStruct);
+#endif
+	
 protected:
 
 	FInstancedStruct(const UScriptStruct* InScriptStruct, uint8* InStructMemory)
diff --git a/Engine/Plugins/Experimental/StructUtils/Source/StructUtils/Public/InstancedStructContainer.h b/Engine/Plugins/Experimental/StructUtils/Source/StructUtils/Public/InstancedStructContainer.h
--- a/Engine/Plugins/Experimental/StructUtils/Source/StructUtils/Public/InstancedStructContainer.h
+++ b/Engine/Plugins/Experimental/StructUtils/Source/StructUtils/Public/InstancedStructContainer.h
@@ -6,6 +6,7 @@
 #include "InstancedStructContainer.generated.h"
 
 struct FInstancedStruct;
+class FReferenceCollector;
 
 /**
  * Array of heterogeneous structs. Can be used as a property, supports serialization,
@@ -222,7 +223,7 @@ public:
 	FConstIterator end() const { return FConstIterator(*this, NumItems); }
 
 	/** Type traits */
-	void AddStructReferencedObjects(class FReferenceCollector& Collector) const;
+	void AddStructReferencedObjects(FReferenceCollector& Collector);
 	bool Identical(const FInstancedStructContainer* Other, uint32 PortFlags) const;
 	bool Serialize(FArchive& Ar);
 	void GetPreloadDependencies(TArray<UObject*>& OutDeps) const;
diff --git a/Engine/Plugins/Experimental/StructUtils/Source/StructUtils/Public/PropertyBag.h b/Engine/Plugins/Experimental/StructUtils/Source/StructUtils/Public/PropertyBag.h
--- a/Engine/Plugins/Experimental/StructUtils/Source/StructUtils/Public/PropertyBag.h
+++ b/Engine/Plugins/Experimental/StructUtils/Source/StructUtils/Public/PropertyBag.h
@@ -593,6 +593,7 @@ struct STRUCTUTILS_API FInstancedPropertyBag
 	TValueOrError<const FPropertyBagArrayRef, EPropertyBagResult> GetArrayRef(const FName Name) const;
 
 	bool Serialize(FArchive& Ar);
+	void AddStructReferencedObjects(FReferenceCollector& Collector);
 
 protected:
 	const void* GetValueAddress(const FPropertyBagPropertyDesc* Desc) const;
@@ -606,7 +607,8 @@ template<> struct TStructOpsTypeTraits<FInstancedPropertyBag> : public TStructOp
 {
 	enum
 	{
-		WithSerializer = true
+		WithSerializer = true,
+		WithAddStructReferencedObjects = true,
 	};
 };
 
@@ -848,6 +850,11 @@ public:
 	/** @return property description based on name. */
 	const FPropertyBagPropertyDesc* FindPropertyDescByName(const FName Name) const;
 
+#if WITH_ENGINE && WITH_EDITOR
+	/** @return true if any of the properties on the bag has type of the specified user defined struct. */
+	bool ContainsUserDefinedStruct(const UUserDefinedStruct* UserDefinedStruct) const;
+#endif	
+	
 protected:
 
 	virtual void InitializeStruct(void* Dest, int32 ArrayDim = 1) const override;
diff --git a/Engine/Plugins/Experimental/StructUtils/Source/StructUtils/Public/StructUtilsTypes.h b/Engine/Plugins/Experimental/StructUtils/Source/StructUtils/Public/StructUtilsTypes.h
--- a/Engine/Plugins/Experimental/StructUtils/Source/StructUtils/Public/StructUtilsTypes.h
+++ b/Engine/Plugins/Experimental/StructUtils/Source/StructUtils/Public/StructUtilsTypes.h
@@ -15,6 +15,7 @@ struct FStructView;
 struct FConstStructView;
 struct FInstancedStruct;
 class FReferenceCollector;
+class UUserDefinedStruct;
 
 namespace UE::StructUtils
 {
@@ -93,3 +94,28 @@ struct FStructTypeSortOperator
 		return FScriptStructSortOperator()(*AScriptStruct, *BScriptStruct);
 	}
 };
+
+#if WITH_ENGINE && WITH_EDITOR
+namespace UE::StructUtils::Private
+{
+	// Private structs used during user defined struct reinstancing.
+	struct STRUCTUTILS_API FStructureToReinstanceScope
+	{
+		explicit FStructureToReinstanceScope(const UUserDefinedStruct* StructureToReinstance);
+		~FStructureToReinstanceScope();
+	private:
+		const UUserDefinedStruct* OldStructureToReinstance = nullptr;
+	};
+
+	struct STRUCTUTILS_API FCurrentReinstanceOuterObjectScope
+	{
+		explicit FCurrentReinstanceOuterObjectScope(UObject* CurrentReinstanceOuterObject);
+		~FCurrentReinstanceOuterObjectScope();
+	private:
+		UObject* OldCurrentReinstanceOuterObject = nullptr;
+	};
+
+	extern STRUCTUTILS_API const UUserDefinedStruct* GetStructureToReinstance();
+	extern STRUCTUTILS_API UObject* GetCurrentReinstanceOuterObject();
+};
+#endif
diff --git a/Engine/Plugins/Experimental/StructUtils/Source/StructUtilsEditor/Private/Customizations/InstancedStructDetails.cpp b/Engine/Plugins/Experimental/StructUtils/Source/StructUtilsEditor/Private/Customizations/InstancedStructDetails.cpp
--- a/Engine/Plugins/Experimental/StructUtils/Source/StructUtilsEditor/Private/Customizations/InstancedStructDetails.cpp
+++ b/Engine/Plugins/Experimental/StructUtils/Source/StructUtilsEditor/Private/Customizations/InstancedStructDetails.cpp
@@ -4,6 +4,7 @@
 #include "DetailWidgetRow.h"
 #include "DetailLayoutBuilder.h"
 #include "IDetailChildrenBuilder.h"
+#include "IPropertyUtilities.h"
 #include "UObject/Package.h"
 #include "Widgets/Images/SImage.h"
 #include "Widgets/Input/SComboButton.h"
@@ -14,11 +15,11 @@
 #include "Styling/SlateIconFinder.h"
 #include "Engine/UserDefinedStruct.h"
 #include "InstancedStruct.h"
-#include "IPropertyUtilities.h"
 #include "Widgets/Layout/SBox.h"
 #include "IStructureDataProvider.h"
 #include "GameFramework/Actor.h"
 #include "Misc/ConfigCacheIni.h"
+#include "StructUtilsDelegates.h"
 
 #define LOCTEXT_NAMESPACE "StructUtilsEditor"
 
@@ -36,6 +37,11 @@ public:
 	
 	virtual ~FInstancedStructProvider() override {}
 
+	void Reset()
+	{
+		StructProperty = nullptr;
+	}
+	
 	virtual bool IsValid() const override
 	{
 		bool bHasValidData = false;
@@ -243,6 +249,21 @@ FInstancedStructDataDetails::FInstancedStructDataDetails(TSharedPtr<IPropertyHan
 	StructProperty = InStructProperty;
 }
 
+FInstancedStructDataDetails::~FInstancedStructDataDetails()
+{
+	UE::StructUtils::Delegates::OnUserDefinedStructReinstanced.Remove(UserDefinedStructReinstancedHandle);
+}
+
+void FInstancedStructDataDetails::OnUserDefinedStructReinstancedHandle(const UUserDefinedStruct& Struct)
+{
+	if (StructProvider.IsValid())
+	{
+		// Reset the struct provider immediately, some update functions might get called with the old struct.
+		StructProvider->Reset();
+	}
+	OnRegenerateChildren.ExecuteIfBound();
+}
+
 void FInstancedStructDataDetails::SetOnRebuildChildren(FSimpleDelegate InOnRegenerateChildren)
 {
 	OnRegenerateChildren = InOnRegenerateChildren;
@@ -284,6 +305,10 @@ void FInstancedStructDataDetails::OnStructHandlePostChange()
 void FInstancedStructDataDetails::GenerateHeaderRowContent(FDetailWidgetRow& NodeRow)
 {
 	StructProperty->SetOnPropertyValueChanged(FSimpleDelegate::CreateSP(this, &FInstancedStructDataDetails::OnStructHandlePostChange));
+	if (!UserDefinedStructReinstancedHandle.IsValid())
+	{
+		UserDefinedStructReinstancedHandle = UE::StructUtils::Delegates::OnUserDefinedStructReinstanced.AddSP(this, &FInstancedStructDataDetails::OnUserDefinedStructReinstancedHandle);
+	}
 }
 
 void FInstancedStructDataDetails::GenerateChildContent(IDetailChildrenBuilder& ChildBuilder)
@@ -352,6 +377,7 @@ void FInstancedStructDetails::CustomizeHeader(TSharedRef<class IPropertyHandle>
 	static const FName NAME_BaseStruct = "BaseStruct";
 	static const FName NAME_StructTypeConst = "StructTypeConst";
 
+	PropUtils = StructCustomizationUtils.GetPropertyUtilities();
 	StructProperty = StructPropertyHandle;
 	PropUtils = StructCustomizationUtils.GetPropertyUtilities();
 
@@ -481,7 +507,7 @@ TSharedRef<SWidget> FInstancedStructDetails::GenerateStructPicker()
 
 	TSharedRef<FInstancedStructFilter> StructFilter = MakeShared<FInstancedStructFilter>();
 	StructFilter->BaseStruct = BaseScriptStruct;
-	StructFilter->bAllowUserDefinedStructs = false;
+	StructFilter->bAllowUserDefinedStructs = true;
 	StructFilter->bAllowBaseStruct = !bExcludeBaseStruct;
 
 	FStructViewerInitializationOptions Options;
diff --git a/Engine/Plugins/Experimental/StructUtils/Source/StructUtilsEditor/Private/Customizations/PropertyBagDetails.cpp b/Engine/Plugins/Experimental/StructUtils/Source/StructUtilsEditor/Private/Customizations/PropertyBagDetails.cpp
--- a/Engine/Plugins/Experimental/StructUtils/Source/StructUtilsEditor/Private/Customizations/PropertyBagDetails.cpp
+++ b/Engine/Plugins/Experimental/StructUtils/Source/StructUtilsEditor/Private/Customizations/PropertyBagDetails.cpp
@@ -439,21 +439,11 @@ void ApplyChangesToPropertyDescs(const FText& SessionName, const TSharedPtr<IPro
 
 bool CanHaveMemberVariableOfType(const FEdGraphPinType& PinType)
 {
-	if ((PinType.PinCategory == UEdGraphSchema_K2::PC_Struct))
-	{
-		if (const UObject* TypeObject = PinType.PinSubCategoryObject.Get())
-		{
-			if (TypeObject->IsA<UUserDefinedStruct>())
-			{
-				return false;
-			}
-		}
-	}
-	else if ((PinType.PinCategory == UEdGraphSchema_K2::PC_Exec) 
-		|| (PinType.PinCategory == UEdGraphSchema_K2::PC_Wildcard)
-		|| (PinType.PinCategory == UEdGraphSchema_K2::PC_MCDelegate)
-		|| (PinType.PinCategory == UEdGraphSchema_K2::PC_Delegate)
-		|| (PinType.PinCategory == UEdGraphSchema_K2::PC_Interface))
+	if (PinType.PinCategory == UEdGraphSchema_K2::PC_Exec 
+		|| PinType.PinCategory == UEdGraphSchema_K2::PC_Wildcard
+		|| PinType.PinCategory == UEdGraphSchema_K2::PC_MCDelegate
+		|| PinType.PinCategory == UEdGraphSchema_K2::PC_Delegate
+		|| PinType.PinCategory == UEdGraphSchema_K2::PC_Interface)
 	{
 		return false;
 	}
diff --git a/Engine/Plugins/Experimental/StructUtils/Source/StructUtilsEditor/Private/StructUtilsEditorModule.cpp b/Engine/Plugins/Experimental/StructUtils/Source/StructUtilsEditor/Private/StructUtilsEditorModule.cpp
--- a/Engine/Plugins/Experimental/StructUtils/Source/StructUtilsEditor/Private/StructUtilsEditorModule.cpp
+++ b/Engine/Plugins/Experimental/StructUtils/Source/StructUtilsEditor/Private/StructUtilsEditorModule.cpp
@@ -4,6 +4,18 @@
 #include "Modules/ModuleManager.h"
 #include "PropertyBagDetails.h"
 #include "PropertyEditorModule.h"
+#include "StructUtilsTypes.h"
+#include "Engine/UserDefinedStruct.h"
+#include "UserDefinedStructure/UserDefinedStructEditorData.h"
+#include "UObject/Package.h"
+#include "UObject/UObjectGlobals.h"
+#include "UObject/UObjectIterator.h"
+#include "CoreGlobals.h"
+#include "GameFramework/Actor.h"
+#include "StructUtilsDelegates.h"
+#include "InstancedStruct.h"
+#include "InstancedStructContainer.h"
+#include "PropertyBag.h"
 
 #define LOCTEXT_NAMESPACE "StructUtilsEditor"
 
@@ -30,4 +42,138 @@ void FStructUtilsEditorModule::ShutdownModule()
 	}
 }
 
+namespace UE::StructUtils::Private
+{
+
+static bool ContainsInstancedStructProperty(const UClass* Class)
+{
+	for (TFieldIterator<FStructProperty> It(Class); It; ++It)
+	{
+		const UScriptStruct* Struct = It->Struct;
+		if (Struct == TBaseStructure<FInstancedStruct>::Get())
+		{
+			return true;
+		}
+		if (Struct == TBaseStructure<FInstancedStructContainer>::Get())
+		{
+			return true;
+		}
+		if (Struct == TBaseStructure<FInstancedPropertyBag>::Get())
+		{
+			return true;
+		}
+	}
+	return false;
+}
+
+static void VisitReferencedObjects(const UUserDefinedStruct* StructToReinstance)
+{
+	// Helper preference collector, does not collect anything, but makes sure AddStructReferencedObjects() gets called e.g. on instanced struct. 
+	class FVisitorReferenceCollector : public FReferenceCollector
+	{
+	public:
+		virtual bool IsIgnoringArchetypeRef() const override { return false; }
+		virtual bool IsIgnoringTransient() const override { return false; }
+
+		virtual void HandleObjectReference(UObject*& Object, const UObject* ReferencingObject, const FProperty* ReferencingProperty) override
+		{
+			// Empty
+		}
+	};
+
+	// Find classes that contain any of the instanced struct types.
+	TArray<const UClass*> InstancedStructClasses;
+	for (TObjectIterator<UClass> It; It; ++It)
+	{
+		if (ContainsInstancedStructProperty(*It))
+		{
+			InstancedStructClasses.Add(*It);
+		}
+	}	
+
+	// Find objects that contain any of the instanced struct types.
+	TArray<UObject*> SourceObjects;
+	for (const UClass* Class : InstancedStructClasses)
+	{
+		TArray<UObject*> Objects;
+		GetObjectsOfClass(Class, Objects);
+		SourceObjects.Append(Objects);
+	}
+
+	FVisitorReferenceCollector Collector;
+
+	// This sets global variable which read in the AddStructReferencedObjects().
+	UE::StructUtils::Private::FStructureToReinstanceScope StructureToReinstanceScope(StructToReinstance);
+
+	for (UObject* Object : SourceObjects)
+	{
+		// This sets global variable which read in the AddStructReferencedObjects().
+		UE::StructUtils::Private::FCurrentReinstanceOuterObjectScope CurrentReinstanceOuterObjectScope(Object);
+		
+		Collector.AddPropertyReferences(Object->GetClass(), Object);
+
+		// AddPropertyReferences() for objects does not handle ARO, do it manually.
+		for (TPropertyValueIterator<FStructProperty> It(Object->GetClass(), Object); It; ++It)
+		{
+			const UScriptStruct* Struct = It.Key()->Struct;
+			void* Instance = const_cast<void*>(It.Value());
+			if (Struct && Struct->StructFlags & STRUCT_AddStructReferencedObjects)
+			{
+				Struct->GetCppStructOps()->AddStructReferencedObjects()(Instance, Collector);
+			}
+		}
+	}
+};
+
+}; // UE::StructUtils::Private
+
+void FStructUtilsEditorModule::PreChange(const UUserDefinedStruct* StructToReinstance, FStructureEditorUtils::EStructureEditorChangeInfo Info)
+{
+	if (!StructToReinstance)
+	{
+		return;
+	}
+
+	// Make a duplicate of the existing struct, and point all instances of the struct to point to the duplicate.
+	// This is done because the original struct will be changed.
+	UUserDefinedStruct* DuplicatedStruct = nullptr;
+	{
+		const FString ReinstancedName = FString::Printf(TEXT("STRUCT_REINST_%s"), *StructToReinstance->GetName());
+		const FName UniqueName = MakeUniqueObjectName(GetTransientPackage(), UUserDefinedStruct::StaticClass(), FName(*ReinstancedName));
+
+		TGuardValue<FIsDuplicatingClassForReinstancing, bool> IsDuplicatingClassForReinstancing(GIsDuplicatingClassForReinstancing, true);
+		DuplicatedStruct = (UUserDefinedStruct*)StaticDuplicateObject(StructToReinstance, GetTransientPackage(), UniqueName, ~RF_Transactional); 
+
+		DuplicatedStruct->Guid = StructToReinstance->Guid;
+		DuplicatedStruct->Bind();
+		DuplicatedStruct->StaticLink(true);
+		DuplicatedStruct->PrimaryStruct = const_cast<UUserDefinedStruct*>(StructToReinstance);
+		DuplicatedStruct->Status = EUserDefinedStructureStatus::UDSS_Duplicate;
+		DuplicatedStruct->SetFlags(RF_Transient);
+		DuplicatedStruct->AddToRoot();
+	}
+
+	UUserDefinedStructEditorData* DuplicatedEditorData = CastChecked<UUserDefinedStructEditorData>(DuplicatedStruct->EditorData);
+	DuplicatedEditorData->RecreateDefaultInstance();
+
+	UE::StructUtils::Private::VisitReferencedObjects(DuplicatedStruct);
+	
+	DuplicatedStruct->RemoveFromRoot();
+}
+
+void FStructUtilsEditorModule::PostChange(const UUserDefinedStruct* StructToReinstance, FStructureEditorUtils::EStructureEditorChangeInfo Info)
+{
+	if (!StructToReinstance)
+	{
+		return;
+	}
+
+	UE::StructUtils::Private::VisitReferencedObjects(StructToReinstance);
+
+	if (UE::StructUtils::Delegates::OnUserDefinedStructReinstanced.IsBound())
+	{
+		UE::StructUtils::Delegates::OnUserDefinedStructReinstanced.Broadcast(*StructToReinstance);
+	}
+}
+
 #undef LOCTEXT_NAMESPACE
diff --git a/Engine/Plugins/Experimental/StructUtils/Source/StructUtilsEditor/Public/InstancedStructDetails.h b/Engine/Plugins/Experimental/StructUtils/Source/StructUtilsEditor/Public/InstancedStructDetails.h
--- a/Engine/Plugins/Experimental/StructUtils/Source/StructUtilsEditor/Public/InstancedStructDetails.h
+++ b/Engine/Plugins/Experimental/StructUtils/Source/StructUtilsEditor/Public/InstancedStructDetails.h
@@ -8,8 +8,10 @@
 
 class IPropertyHandle;
 class IDetailPropertyRow;
+class IPropertyHandle;
 class FStructOnScope;
 class SWidget;
+class SComboButton;
 struct FInstancedStruct;
 class FInstancedStructProvider;
 
@@ -45,8 +47,7 @@ private:
 	/** The base struct that we're allowing to be picked (controlled by the "BaseStruct" meta-data) */
 	UScriptStruct* BaseScriptStruct = nullptr;
 
-	TSharedPtr<class SComboButton> ComboButton;
-
+	TSharedPtr<SComboButton> ComboButton;
 	TSharedPtr<IPropertyUtilities> PropUtils;
 	
 	FDelegateHandle OnObjectsReinstancedHandle;
@@ -62,6 +63,7 @@ class STRUCTUTILSEDITOR_API FInstancedStructDataDetails : public IDetailCustomNo
 {
 public:
 	FInstancedStructDataDetails(TSharedPtr<IPropertyHandle> InStructProperty);
+	~FInstancedStructDataDetails();
 
 	//~ Begin IDetailCustomNodeBuilder interface
 	virtual void SetOnRebuildChildren(FSimpleDelegate InOnRegenerateChildren) override;
@@ -81,6 +83,11 @@ public:
 	virtual void OnChildRowAdded(IDetailPropertyRow& ChildRow) {}
 
 private:
+	void OnUserDefinedStructReinstancedHandle(const class UUserDefinedStruct& Struct);
+
+	/** Pre/Post change notifications for struct value changes */
+	void OnStructValuePreChange();
+	void OnStructValuePostChange();
 	void OnStructHandlePostChange();
 
 	/** Returns type of the instanced struct for each instance/object being edited. */
@@ -103,6 +110,8 @@ private:
 
 	/** True if we're currently handling a StructValuePostChange */
 	bool bIsHandlingStructValuePostChange = false;
+	
+	FDelegateHandle UserDefinedStructReinstancedHandle;
 };
 
 #if UE_ENABLE_INCLUDE_ORDER_DEPRECATED_IN_5_2
diff --git a/Engine/Plugins/Experimental/StructUtils/Source/StructUtilsEditor/Public/StructUtilsEditorModule.h b/Engine/Plugins/Experimental/StructUtils/Source/StructUtilsEditor/Public/StructUtilsEditorModule.h
--- a/Engine/Plugins/Experimental/StructUtils/Source/StructUtilsEditor/Public/StructUtilsEditorModule.h
+++ b/Engine/Plugins/Experimental/StructUtils/Source/StructUtilsEditor/Public/StructUtilsEditorModule.h
@@ -3,20 +3,28 @@
 #pragma  once
 
 #include "Modules/ModuleInterface.h"
-
+#include "Kismet2/StructureEditorUtils.h"
 
 class IStructUtilsEditor;
 struct FGraphPanelNodeFactory;
+class UUserDefinedStruct;
 
 /**
 * The public interface to this module
 */
-class STRUCTUTILSEDITOR_API FStructUtilsEditorModule : public IModuleInterface//, public IHasMenuExtensibility, public IHasToolBarExtensibility
+class STRUCTUTILSEDITOR_API FStructUtilsEditorModule : public IModuleInterface, public FStructureEditorUtils::INotifyOnStructChanged
 {
 public:
 	// Begin IModuleInterface
 	virtual void StartupModule() override;
 	virtual void ShutdownModule() override;
+
+protected:
+
+	// INotifyOnStructChanged
+	virtual void PreChange(const class UUserDefinedStruct* Struct, FStructureEditorUtils::EStructureEditorChangeInfo Info) override;
+	virtual void PostChange(const class UUserDefinedStruct* Struct, FStructureEditorUtils::EStructureEditorChangeInfo Info) override;
+	// ~INotifyOnStructChanged
 };
 
 #if UE_ENABLE_INCLUDE_ORDER_DEPRECATED_IN_5_2
diff --git a/Engine/Plugins/Experimental/StructUtils/Source/StructUtilsEditor/StructUtilsEditor.Build.cs b/Engine/Plugins/Experimental/StructUtils/Source/StructUtilsEditor/StructUtilsEditor.Build.cs
--- a/Engine/Plugins/Experimental/StructUtils/Source/StructUtilsEditor/StructUtilsEditor.Build.cs
+++ b/Engine/Plugins/Experimental/StructUtils/Source/StructUtilsEditor/StructUtilsEditor.Build.cs
@@ -9,37 +9,40 @@ namespace UnrealBuildTool.Rules
 			PCHUsage = ModuleRules.PCHUsageMode.UseExplicitOrSharedPCHs;
 
 			PublicIncludePaths.AddRange(
-			new string[] {
-			}
+				new string[] {
+				}
 			);
 
 			PublicDependencyModuleNames.AddRange(
-			new string[] {
-				"Core",
-				"CoreUObject",
-				"Engine",
-				"InputCore",
-				"AssetTools",
-				"UnrealEd",
-				"Slate",
-				"SlateCore",				
-				"PropertyEditor",
-				"AIModule",
-				"StructUtils",
-				"DetailCustomizations",
-				"ComponentVisualizers",
-			}
+				new string[] {
+					"Core",
+					"CoreUObject",
+					"Engine",
+					"InputCore",
+					"AssetTools",
+					"UnrealEd",
+					"Slate",
+					"SlateCore",				
+					"PropertyEditor",
+					"AIModule",
+					"StructUtils",
+					"StructUtilsEngine",
+					"DetailCustomizations",
+					"ComponentVisualizers",
+				}
 			);
 
 			PrivateDependencyModuleNames.AddRange(
-			new string[] {
-				"RenderCore",
-				"GraphEditor",
-				"KismetWidgets",
-				"AIGraph",
-				"ToolMenus",
-				"BlueprintGraph",
-			}
+				new string[] {
+					"RenderCore",
+					"GraphEditor",
+					"KismetWidgets",
+					"PropertyEditor",
+					"AIGraph",
+					"ToolMenus",
+					"BlueprintGraph",
+					"KismetCompiler",
+				}
 			);
 		}
 
diff --git a/Engine/Plugins/Experimental/StructUtils/Source/StructUtilsEngine/Private/StructUtilsDelegates.cpp b/Engine/Plugins/Experimental/StructUtils/Source/StructUtilsEngine/Private/StructUtilsDelegates.cpp
new file mode 100644
--- /dev/null
+++ b/Engine/Plugins/Experimental/StructUtils/Source/StructUtilsEngine/Private/StructUtilsDelegates.cpp
@@ -0,0 +1,10 @@
+﻿// Copyright Epic Games, Inc. All Rights Reserved.
+
+#include "StructUtilsDelegates.h"
+
+namespace UE::StructUtils::Delegates
+{
+#if WITH_EDITOR
+	FOnUserDefinedStructReinstanced OnUserDefinedStructReinstanced;
+#endif
+} // UE::StructUtils::Delegates
diff --git a/Engine/Plugins/Experimental/StructUtils/Source/StructUtilsEngine/Public/StructUtilsDelegates.h b/Engine/Plugins/Experimental/StructUtils/Source/StructUtilsEngine/Public/StructUtilsDelegates.h
new file mode 100644
--- /dev/null
+++ b/Engine/Plugins/Experimental/StructUtils/Source/StructUtilsEngine/Public/StructUtilsDelegates.h
@@ -0,0 +1,16 @@
+﻿// Copyright Epic Games, Inc. All Rights Reserved.
+
+#pragma once
+
+#include "Delegates/Delegate.h"
+
+class UUserDefinedStruct;
+
+namespace UE::StructUtils::Delegates
+{
+#if WITH_EDITOR
+	/** Called after the FInstancedStructs has been reinstantiated. E.g. safe to update UI. */
+	DECLARE_MULTICAST_DELEGATE_OneParam(FOnUserDefinedStructReinstanced, const UUserDefinedStruct& /*UserDefinedStruct*/);
+	extern STRUCTUTILSENGINE_API FOnUserDefinedStructReinstanced OnUserDefinedStructReinstanced;
+#endif
+} // UE::StructUtils::Delegates
diff --git a/Engine/Plugins/Runtime/StateTree/Source/StateTreeEditorModule/Public/StateTreeEditorModule.h b/Engine/Plugins/Runtime/StateTree/Source/StateTreeEditorModule/Public/StateTreeEditorModule.h
--- a/Engine/Plugins/Runtime/StateTree/Source/StateTreeEditorModule/Public/StateTreeEditorModule.h
+++ b/Engine/Plugins/Runtime/StateTree/Source/StateTreeEditorModule/Public/StateTreeEditorModule.h
@@ -7,6 +7,7 @@
 #include "Toolkits/AssetEditorToolkit.h"
 
 class UStateTree;
+class UUserDefinedStruct;
 class IStateTreeEditor;
 struct FStateTreeNodeClassCache;
 
@@ -34,6 +35,8 @@ protected:
 	TSharedPtr<FExtensibilityManager> MenuExtensibilityManager;
 	TSharedPtr<FExtensibilityManager> ToolBarExtensibilityManager;
 	TSharedPtr<FStateTreeNodeClassCache> NodeClassCache;
+
+	FDelegateHandle OnUserDefinedStructReinstancedHandle;
 };
 
 #if UE_ENABLE_INCLUDE_ORDER_DEPRECATED_IN_5_2
diff --git a/Engine/Plugins/Runtime/StateTree/Source/StateTreeModule/Private/StateTree.cpp b/Engine/Plugins/Runtime/StateTree/Source/StateTreeModule/Private/StateTree.cpp
--- a/Engine/Plugins/Runtime/StateTree/Source/StateTreeModule/Private/StateTree.cpp
+++ b/Engine/Plugins/Runtime/StateTree/Source/StateTreeModule/Private/StateTree.cpp
@@ -10,7 +10,10 @@
 #include "StateTreeDelegates.h"
 #include "Logging/LogScopedVerbosityOverride.h"
 #include "Misc/DataValidation.h"
-
+#include "StructUtilsDelegates.h"
+#if WITH_EDITOR
+#include "Engine/UserDefinedStruct.h"
+#endif
 #include UE_INLINE_GENERATED_CPP_BY_NAME(StateTree)
 
 const FGuid FStateTreeCustomVersion::GUID(0x28E21331, 0x501F4723, 0x8110FA64, 0xEA10DA1E);
@@ -168,11 +171,29 @@ void UStateTree::OnObjectsReinstanced(const FReplacementObjectMap& ObjectMap)
 	}
 }
 
+void UStateTree::OnUserDefinedStructReinstanced(const UUserDefinedStruct& UserDefinedStruct)
+{
+	// Struct utils handle reinstancing the struct values (instanced struct, property bag, etc).
+	// We will need to update the property binding.
+	
+	TSet<const UStruct*> Structs;
+	Structs.Add(&UserDefinedStruct);
+	
+	if (PropertyBindings.ContainsAnyStruct(Structs))
+	{
+		if (!Link())
+		{
+			UE_LOG(LogStateTree, Error, TEXT("%s failed to link after Struct reinstantiation. Take a look at the asset for any errors. Asset will not be usable at runtime."), *GetFullName());
+		}
+	}
+}
+
 void UStateTree::PostInitProperties()
 {
 	Super::PostInitProperties();
 	
 	OnObjectsReinstancedHandle = FCoreUObjectDelegates::OnObjectsReinstanced.AddUObject(this, &UStateTree::OnObjectsReinstanced);
+	OnUserDefinedStructReinstancedHandle = UE::StructUtils::Delegates::OnUserDefinedStructReinstanced.AddUObject(this, &UStateTree::OnUserDefinedStructReinstanced);
 }
 
 void UStateTree::BeginDestroy()
@@ -182,6 +203,11 @@ void UStateTree::BeginDestroy()
 		FCoreUObjectDelegates::OnObjectsReinstanced.Remove(OnObjectsReinstancedHandle);
 		OnObjectsReinstancedHandle.Reset();
 	}
+	if (OnUserDefinedStructReinstancedHandle.IsValid())
+	{
+		UE::StructUtils::Delegates::OnUserDefinedStructReinstanced.Remove(OnUserDefinedStructReinstancedHandle);
+		OnUserDefinedStructReinstancedHandle.Reset();
+	}
 	
 	Super::BeginDestroy();
 }
diff --git a/Engine/Plugins/Runtime/StateTree/Source/StateTreeModule/Public/StateTree.h b/Engine/Plugins/Runtime/StateTree/Source/StateTreeModule/Public/StateTree.h
--- a/Engine/Plugins/Runtime/StateTree/Source/StateTreeModule/Public/StateTree.h
+++ b/Engine/Plugins/Runtime/StateTree/Source/StateTreeModule/Public/StateTree.h
@@ -8,6 +8,7 @@
 #include "StateTreeInstanceData.h"
 #include "StateTree.generated.h"
 
+class UUserDefinedStruct;
 
 /** Custom serialization version for StateTree Asset */
 struct STATETREEMODULE_API FStateTreeCustomVersion
@@ -134,6 +135,7 @@ public:
 	TObjectPtr<UObject> EditorData;
 
 	FDelegateHandle OnObjectsReinstancedHandle;
+	FDelegateHandle OnUserDefinedStructReinstancedHandle;
 #endif
 
 	/** Hash of the editor data from last compile. Also used to detect mismatching events from recorded traces. */
@@ -157,6 +159,7 @@ protected:
 #if WITH_EDITOR
 	using FReplacementObjectMap = TMap<UObject*, UObject*>;
 	void OnObjectsReinstanced(const FReplacementObjectMap& ObjectMap);
+	void OnUserDefinedStructReinstanced(const UUserDefinedStruct& UserDefinedStruct);
 	virtual void PostInitProperties() override;
 	virtual void BeginDestroy() override;
 	virtual void GetAssetRegistryTags(TArray<FAssetRegistryTag>& OutTags) const override;
diff --git a/Engine/Plugins/Runtime/StateTree/Source/StateTreeModule/StateTreeModule.Build.cs b/Engine/Plugins/Runtime/StateTree/Source/StateTreeModule/StateTreeModule.Build.cs
--- a/Engine/Plugins/Runtime/StateTree/Source/StateTreeModule/StateTreeModule.Build.cs
+++ b/Engine/Plugins/Runtime/StateTree/Source/StateTreeModule/StateTreeModule.Build.cs
@@ -20,6 +20,7 @@ namespace UnrealBuildTool.Rules
 					"AIModule",
 					"GameplayTags",
 					"StructUtils",
+					"StructUtilsEngine",
 				}
 			);
 
