StructUtils: Fix merge error.

diff --git a/Engine/Plugins/Experimental/StructUtils/Source/StructUtils/Private/InstancedStructContainer.cpp b/Engine/Plugins/Experimental/StructUtils/Source/StructUtils/Private/InstancedStructContainer.cpp
--- a/Engine/Plugins/Experimental/StructUtils/Source/StructUtils/Private/InstancedStructContainer.cpp
+++ b/Engine/Plugins/Experimental/StructUtils/Source/StructUtils/Private/InstancedStructContainer.cpp
@@ -582,7 +582,7 @@ bool FInstancedStructContainer::Serialize(FArchive& Ar)
 				const FItem& Item = GetItem(Index);
 				
 #if WITH_ENGINE && WITH_EDITOR
-				UUserDefinedStruct* UserDefinedStruct = Cast<UUserDefinedStruct>(ConstCast(Item.ScriptStruct);
+				UUserDefinedStruct* UserDefinedStruct = Cast<UUserDefinedStruct>(ConstCast(Item.ScriptStruct));
 				if (UserDefinedStruct
 					&& UserDefinedStruct->Status == EUserDefinedStructureStatus::UDSS_Duplicate
 					&& UserDefinedStruct->PrimaryStruct.IsValid())
