StructUtils: Do not allow User Defined Structs in instanced struct picker if BaseStruct is set.

diff --git a/Engine/Plugins/Experimental/StructUtils/Source/StructUtilsEditor/Private/Customizations/InstancedStructDetails.cpp b/Engine/Plugins/Experimental/StructUtils/Source/StructUtilsEditor/Private/Customizations/InstancedStructDetails.cpp
--- a/Engine/Plugins/Experimental/StructUtils/Source/StructUtilsEditor/Private/Customizations/InstancedStructDetails.cpp
+++ b/Engine/Plugins/Experimental/StructUtils/Source/StructUtilsEditor/Private/Customizations/InstancedStructDetails.cpp
@@ -507,7 +507,7 @@ TSharedRef<SWidget> FInstancedStructDetails::GenerateStructPicker()
 
 	TSharedRef<FInstancedStructFilter> StructFilter = MakeShared<FInstancedStructFilter>();
 	StructFilter->BaseStruct = BaseScriptStruct;
-	StructFilter->bAllowUserDefinedStructs = true;
+	StructFilter->bAllowUserDefinedStructs = BaseScriptStruct == nullptr; // Only allow user defined structs when BaseStruct is not set.
 	StructFilter->bAllowBaseStruct = !bExcludeBaseStruct;
 
 	FStructViewerInitializationOptions Options;
