StructUtils: Fix User Defined Struct instancing on CDOs
- Handle instanced UDS in CDOs and sparse class data
- Fixed Property Bag ref couting when replacing a bag during UDS reinstancing
- Visit all UObjects instead of just where we find instanced structs (some cases are by value, so we cannot prune by class)
- Use AddPropertyReferencesWithStructARO() to visit data, removed manual ARO hack

diff --git a/Engine/Plugins/Experimental/StructUtils/Source/StructUtils/Private/InstancedStruct.cpp b/Engine/Plugins/Experimental/StructUtils/Source/StructUtils/Private/InstancedStruct.cpp
--- a/Engine/Plugins/Experimental/StructUtils/Source/StructUtils/Private/InstancedStruct.cpp
+++ b/Engine/Plugins/Experimental/StructUtils/Source/StructUtils/Private/InstancedStruct.cpp
@@ -458,7 +458,10 @@ void FInstancedStruct::AddStructReferencedObjects(FReferenceCollector& Collector
 				{
 					if (UObject* Outer = UE::StructUtils::Private::GetCurrentReinstanceOuterObject())
 					{
-						Outer->MarkPackageDirty();
+						if (!Outer->IsA<UClass>() && !Outer->HasAnyFlags(RF_ClassDefaultObject))
+						{
+							Outer->MarkPackageDirty();
+						}
 					}
 
 					TArray<uint8> Data;
diff --git a/Engine/Plugins/Experimental/StructUtils/Source/StructUtils/Private/InstancedStructContainer.cpp b/Engine/Plugins/Experimental/StructUtils/Source/StructUtils/Private/InstancedStructContainer.cpp
--- a/Engine/Plugins/Experimental/StructUtils/Source/StructUtils/Private/InstancedStructContainer.cpp
+++ b/Engine/Plugins/Experimental/StructUtils/Source/StructUtils/Private/InstancedStructContainer.cpp
@@ -418,7 +418,10 @@ void FInstancedStructContainer::AddStructReferencedObjects(FReferenceCollector&
 				
 				if (UObject* Outer = UE::StructUtils::Private::GetCurrentReinstanceOuterObject())
 				{
-					Outer->MarkPackageDirty();
+					if (!Outer->IsA<UClass>() && !Outer->HasAnyFlags(RF_ClassDefaultObject))
+					{
+						Outer->MarkPackageDirty();
+					}
 				}
 				
 				TArray<uint8> Data;
diff --git a/Engine/Plugins/Experimental/StructUtils/Source/StructUtils/Private/PropertyBag.cpp b/Engine/Plugins/Experimental/StructUtils/Source/StructUtils/Private/PropertyBag.cpp
--- a/Engine/Plugins/Experimental/StructUtils/Source/StructUtils/Private/PropertyBag.cpp
+++ b/Engine/Plugins/Experimental/StructUtils/Source/StructUtils/Private/PropertyBag.cpp
@@ -1883,6 +1883,10 @@ void FInstancedPropertyBag::AddStructReferencedObjects(FReferenceCollector& Coll
 
 				const UPropertyBag* NewBag = UPropertyBag::GetOrCreateFromDescs(PropertyDescs);
 				Value.ReplaceScriptStructInternal(NewBag);
+
+				// Adjust recount manually, since we replaced the struct above.
+				Bag->DecrementRefCount();
+				NewBag->IncrementRefCount();
 			}
 			else
 			{
@@ -1892,7 +1896,10 @@ void FInstancedPropertyBag::AddStructReferencedObjects(FReferenceCollector& Coll
 
 				if (UObject* Outer = UE::StructUtils::Private::GetCurrentReinstanceOuterObject())
 				{
-					Outer->MarkPackageDirty();
+					if (!Outer->IsA<UClass>() && !Outer->HasAnyFlags(RF_ClassDefaultObject))
+					{
+						Outer->MarkPackageDirty();
+					}
 				}
 				
 				TArray<uint8> Data;
@@ -2295,10 +2302,26 @@ bool UPropertyBag::ContainsUserDefinedStruct(const UUserDefinedStruct* UserDefin
 }
 #endif
 
-void UPropertyBag::InitializeStruct(void* Dest, int32 ArrayDim) const
+void UPropertyBag::DecrementRefCount() const
 {
-	Super::InitializeStruct(Dest, ArrayDim);
+	// Do ref counting based on struct usage.
+	// This ensures that if the UPropertyBag is still valid in the C++ destructor of
+	// the last instance of the bag.
 
+	UPropertyBag* NonConstThis = const_cast<UPropertyBag*>(this);
+	const int32 OldCount = NonConstThis->RefCount.fetch_sub(1, std::memory_order_acq_rel);
+	if (OldCount == 1)
+	{
+		NonConstThis->RemoveFromRoot();
+	}
+	if (OldCount <= 0)
+	{
+		UE_LOG(LogCore, Error, TEXT("PropertyBag: DestroyStruct is called when RefCount is %d."), OldCount);
+	}
+}
+
+void UPropertyBag::IncrementRefCount() const
+{
 	// Do ref counting based on struct usage.
 	// This ensures that if the UPropertyBag is still valid in the C++ destructor of
 	// the last instance of the bag.
@@ -2310,24 +2333,18 @@ void UPropertyBag::InitializeStruct(void* Dest, int32 ArrayDim) const
 	}
 }
 
+void UPropertyBag::InitializeStruct(void* Dest, int32 ArrayDim) const
+{
+	Super::InitializeStruct(Dest, ArrayDim);
+
+	IncrementRefCount();
+}
+
 void UPropertyBag::DestroyStruct(void* Dest, int32 ArrayDim) const
 {
 	Super::DestroyStruct(Dest, ArrayDim);
 
-	// Do ref counting based on struct usage.
-	// This ensures that if the UPropertyBag is still valid in the C++ destructor of
-	// the last instance of the bag.
-
-	UPropertyBag* NonConstThis = const_cast<UPropertyBag*>(this);
-	const int32 OldCount = NonConstThis->RefCount.fetch_sub(1, std::memory_order_acq_rel);
-	if (OldCount == 1)
-	{
-		NonConstThis->RemoveFromRoot();
-	}
-	if (OldCount <= 0)
-	{
-		UE_LOG(LogCore, Error, TEXT("PropertyBag: DestroyStruct is called when RefCount is %d."), OldCount);
-	}
+	DecrementRefCount();
 }
 
 void UPropertyBag::FinishDestroy()
diff --git a/Engine/Plugins/Experimental/StructUtils/Source/StructUtils/Public/PropertyBag.h b/Engine/Plugins/Experimental/StructUtils/Source/StructUtils/Public/PropertyBag.h
--- a/Engine/Plugins/Experimental/StructUtils/Source/StructUtils/Public/PropertyBag.h
+++ b/Engine/Plugins/Experimental/StructUtils/Source/StructUtils/Public/PropertyBag.h
@@ -857,6 +857,9 @@ public:
 	
 protected:
 
+	void DecrementRefCount() const;
+	void IncrementRefCount() const;
+	
 	virtual void InitializeStruct(void* Dest, int32 ArrayDim = 1) const override;
 	virtual void DestroyStruct(void* Dest, int32 ArrayDim = 1) const override;
 	virtual void FinishDestroy();
diff --git a/Engine/Plugins/Experimental/StructUtils/Source/StructUtilsEditor/Private/StructUtilsEditorModule.cpp b/Engine/Plugins/Experimental/StructUtils/Source/StructUtilsEditor/Private/StructUtilsEditorModule.cpp
--- a/Engine/Plugins/Experimental/StructUtils/Source/StructUtilsEditor/Private/StructUtilsEditorModule.cpp
+++ b/Engine/Plugins/Experimental/StructUtils/Source/StructUtilsEditor/Private/StructUtilsEditorModule.cpp
@@ -45,27 +45,6 @@ void FStructUtilsEditorModule::ShutdownModule()
 namespace UE::StructUtils::Private
 {
 
-static bool ContainsInstancedStructProperty(const UClass* Class)
-{
-	for (TFieldIterator<FStructProperty> It(Class); It; ++It)
-	{
-		const UScriptStruct* Struct = It->Struct;
-		if (Struct == TBaseStructure<FInstancedStruct>::Get())
-		{
-			return true;
-		}
-		if (Struct == TBaseStructure<FInstancedStructContainer>::Get())
-		{
-			return true;
-		}
-		if (Struct == TBaseStructure<FInstancedPropertyBag>::Get())
-		{
-			return true;
-		}
-	}
-	return false;
-}
-
 static void VisitReferencedObjects(const UUserDefinedStruct* StructToReinstance)
 {
 	// Helper preference collector, does not collect anything, but makes sure AddStructReferencedObjects() gets called e.g. on instanced struct. 
@@ -81,46 +60,39 @@ static void VisitReferencedObjects(const UUserDefinedStruct* StructToReinstance)
 		}
 	};
 
-	// Find classes that contain any of the instanced struct types.
-	TArray<const UClass*> InstancedStructClasses;
-	for (TObjectIterator<UClass> It; It; ++It)
-	{
-		if (ContainsInstancedStructProperty(*It))
-		{
-			InstancedStructClasses.Add(*It);
-		}
-	}	
-
-	// Find objects that contain any of the instanced struct types.
-	TArray<UObject*> SourceObjects;
-	for (const UClass* Class : InstancedStructClasses)
-	{
-		TArray<UObject*> Objects;
-		GetObjectsOfClass(Class, Objects);
-		SourceObjects.Append(Objects);
-	}
-
 	FVisitorReferenceCollector Collector;
 
 	// This sets global variable which read in the AddStructReferencedObjects().
 	UE::StructUtils::Private::FStructureToReinstanceScope StructureToReinstanceScope(StructToReinstance);
 
-	for (UObject* Object : SourceObjects)
+	for (TObjectIterator<UObject> ObjectIt; ObjectIt; ++ObjectIt)
 	{
+		UObject* Object = *ObjectIt;
+		
 		// This sets global variable which read in the AddStructReferencedObjects().
 		UE::StructUtils::Private::FCurrentReinstanceOuterObjectScope CurrentReinstanceOuterObjectScope(Object);
 		
-		Collector.AddPropertyReferences(Object->GetClass(), Object);
+		Collector.AddPropertyReferencesWithStructARO(Object->GetClass(), Object);
+	}
+
+	// Handle CDOs and sparse class data 
+	for (TObjectIterator<UClass> ClassIt; ClassIt; ++ClassIt)
+	{
+		UClass* Class = *ClassIt;
+	
+		// Handle Sparse class data
+		if (void* SparseData = const_cast<void*>(Class->GetSparseClassData(EGetSparseClassDataMethod::ReturnIfNull)))
+		{
+			UE::StructUtils::Private::FCurrentReinstanceOuterObjectScope CurrentReinstanceOuterObjectScope(Class);
+			const UScriptStruct* SparseDataStruct = Class->GetSparseClassDataStruct();
+			Collector.AddPropertyReferencesWithStructARO(SparseDataStruct, SparseData);
+		}
 
-		// AddPropertyReferences() for objects does not handle ARO, do it manually.
-		for (TPropertyValueIterator<FStructProperty> It(Object->GetClass(), Object); It; ++It)
+		// Handle CDO
+		if (UObject* CDO = Class->GetDefaultObject())
 		{
-			const UScriptStruct* Struct = It.Key()->Struct;
-			void* Instance = const_cast<void*>(It.Value());
-			if (Struct && Struct->StructFlags & STRUCT_AddStructReferencedObjects)
-			{
-				Struct->GetCppStructOps()->AddStructReferencedObjects()(Instance, Collector);
-			}
+			UE::StructUtils::Private::FCurrentReinstanceOuterObjectScope CurrentReinstanceOuterObjectScope(CDO);
+			Collector.AddPropertyReferencesWithStructARO(Class, CDO);
 		}
 	}
 };
