[StructUtils] Add callback options with metadata on FInstancedPropertyBag

* Add a TypeFilter callback to filter out some types when changing the type of a property
* Add a RemoveCheck callback to cancel the removal of a property
* Add generic way to invoke UFunctions. Will need to be moved in its seperate file later on.

diff --git a/Engine/Plugins/Experimental/StructUtils/Source/StructUtils/Private/StructUtilsMetadata.cpp b/Engine/Plugins/Experimental/StructUtils/Source/StructUtils/Private/StructUtilsMetadata.cpp
--- a/Engine/Plugins/Experimental/StructUtils/Source/StructUtils/Private/StructUtilsMetadata.cpp
+++ b/Engine/Plugins/Experimental/StructUtils/Source/StructUtils/Private/StructUtilsMetadata.cpp
@@ -7,4 +7,6 @@ namespace UE::StructUtils::Metadata
 	const FName FixedLayoutName = TEXT("FixedLayout");
 	const FName DefaultTypeName = TEXT("DefaultType");
 	const FName AllowArraysName = TEXT("AllowArrays");
+	const FName IsPinTypeAcceptedName = TEXT("IsPinTypeAccepted");
+	const FName CanRemovePropertyName = TEXT("CanRemoveProperty");
 }
\ No newline at end of file
diff --git a/Engine/Plugins/Experimental/StructUtils/Source/StructUtils/Public/StructUtilsMetadata.h b/Engine/Plugins/Experimental/StructUtils/Source/StructUtils/Public/StructUtilsMetadata.h
--- a/Engine/Plugins/Experimental/StructUtils/Source/StructUtils/Public/StructUtilsMetadata.h
+++ b/Engine/Plugins/Experimental/StructUtils/Source/StructUtils/Public/StructUtilsMetadata.h
@@ -19,4 +19,10 @@ namespace UE::StructUtils::Metadata
 
 	/// [PropertyMetadata] AllowArrays: By default it is always allowed to create arrays for properties. Use AllowArrays=false to disable Array support.
 	extern STRUCTUTILS_API const FName AllowArraysName;
+
+	/// [PropertyMetadata] IsPinTypeAccepted: Name of a UFunction with signature bool(FEdGraphPinType). Returns false if the type should be discarded.
+	extern STRUCTUTILS_API const FName IsPinTypeAcceptedName;
+
+	/// [PropertyMetadata] CanRemoveProperty: Name of a UFunction with signature bool(FGuid, FName). ID and name of the property that will be removed. Returns false if the property should not be removed.
+	extern STRUCTUTILS_API const FName CanRemovePropertyName;
 }
\ No newline at end of file
diff --git a/Engine/Plugins/Experimental/StructUtils/Source/StructUtilsEditor/Private/Customizations/PropertyBagDetails.cpp b/Engine/Plugins/Experimental/StructUtils/Source/StructUtilsEditor/Private/Customizations/PropertyBagDetails.cpp
--- a/Engine/Plugins/Experimental/StructUtils/Source/StructUtilsEditor/Private/Customizations/PropertyBagDetails.cpp
+++ b/Engine/Plugins/Experimental/StructUtils/Source/StructUtilsEditor/Private/Customizations/PropertyBagDetails.cpp
@@ -15,6 +15,7 @@
 #include "Engine/UserDefinedStruct.h"
 #include "SPinTypeSelector.h"
 #include "StructUtilsMetadata.h"
+#include "Templates/ValueOrError.h"
 
 #include UE_INLINE_GENERATED_CPP_BY_NAME(PropertyBagDetails)
 
@@ -451,6 +452,267 @@ bool CanHaveMemberVariableOfType(const FEdGraphPinType& PinType)
 	return true;
 }
 
+bool FindUserFunction(TSharedPtr<IPropertyHandle> InStructProperty, FName InFuncMetadataName, UFunction*& OutFunc, UObject*& OutTarget)
+{
+	FProperty* MetadataProperty = InStructProperty->GetMetaDataProperty();
+
+	OutFunc = nullptr;
+	OutTarget = nullptr;
+
+	if (!MetadataProperty || !MetadataProperty->HasMetaData(InFuncMetadataName))
+	{
+		return false;
+	}
+
+	FString FunctionName = MetadataProperty->GetMetaData(InFuncMetadataName);
+	if (FunctionName.IsEmpty())
+	{
+		return false;
+	}
+
+	TArray<UObject*> OutObjects;
+	InStructProperty->GetOuterObjects(OutObjects);
+
+	// Check for external function references, taken from GetOptions
+	if (FunctionName.Contains(TEXT(".")))
+	{
+		OutFunc = FindObject<UFunction>(nullptr, *FunctionName, true);
+
+		if (ensureMsgf(OutFunc && OutFunc->HasAnyFunctionFlags(EFunctionFlags::FUNC_Static), TEXT("[%s] Didn't find function %s or expected it to be static"), *InFuncMetadataName.ToString(), *FunctionName))
+		{
+			UObject* GetOptionsCDO = OutFunc->GetOuterUClass()->GetDefaultObject();
+			OutTarget = GetOptionsCDO;
+		}
+	}
+	else if (OutObjects.Num() > 0)
+	{
+		OutTarget = OutObjects[0];
+		OutFunc = OutTarget->GetClass() ? OutTarget->GetClass()->FindFunctionByName(*FunctionName) : nullptr;
+	}
+
+	// Only support native functions
+	if (!ensureMsgf(OutFunc && OutFunc->IsNative(), TEXT("[%s] Didn't find function %s or expected it to be native"), *InFuncMetadataName.ToString(), *FunctionName))
+	{
+		OutFunc = nullptr;
+		OutTarget = nullptr;
+	}
+
+	return OutTarget != nullptr && OutFunc != nullptr;
+}
+
+// UFunction calling helpers. 
+// Use our "own" hardcoded reflection system for types used in UFunctions calls in this file.
+template<typename T> struct TypeName { static const TCHAR* Get() = delete; };
+
+#define DEFINE_TYPENAME(InType) template<> struct TypeName<InType> { static const TCHAR *Get() { return TEXT(#InType); }};
+DEFINE_TYPENAME(bool)
+DEFINE_TYPENAME(FGuid)
+DEFINE_TYPENAME(FName)
+DEFINE_TYPENAME(FEdGraphPinType)
+#undef DEFINE_TYPENAME
+
+// Wrapper around a param that store an address (const_cast for const ptr, be careful of that)
+// a string identifiying the underlying cpp type and if the input value is const, mark it const.
+struct FFuncParam
+{
+	void* Value = nullptr;
+	const TCHAR* CppType;
+	bool bIsConst = false;
+
+	template<typename T>
+	static FFuncParam Make(T& Value)
+	{
+		FFuncParam Result;
+		Result.Value = &Value;
+		Result.CppType = TypeName<T>::Get();
+		return Result;
+	}
+
+	template<typename T>
+	static FFuncParam Make(const T& Value)
+	{
+		FFuncParam Result;
+		Result.Value = &const_cast<T&>(Value);
+		Result.CppType = TypeName<T>::Get();
+		Result.bIsConst = true;
+		return Result;
+	}
+};
+
+// Validate that the function pass as parameter has signature ReturnType(ArgsTypes...)
+template <typename ReturnType, typename... ArgsTypes>
+bool ValidateFunctionSignature(UFunction* InFunc)
+{
+	if (!InFunc)
+	{
+		return false;
+	}
+
+	constexpr int32 NumParms = std::is_same_v <ReturnType, void> ? sizeof...(ArgsTypes) : (sizeof...(ArgsTypes) + 1);
+
+	if (NumParms != InFunc->NumParms)
+	{
+		return false;
+	}
+
+	const TCHAR* ArgsCppTypes[NumParms] = { TypeName<ArgsTypes>::Get() ... };
+	
+	// If we have a return type, put it at the end. UFunction will have the return type after InArgs in the field iterator.
+	if constexpr (!std::is_same_v<ReturnType, void>)
+	{
+		ArgsCppTypes[NumParms - 1] = TypeName<ReturnType>::Get();
+	}
+	else
+	{
+		// Otherwise, check that the function doesn't have a return param
+		if (InFunc->GetReturnProperty() != nullptr)
+		{
+			return false;
+		}
+	}
+
+	int32 ArgCppTypesIndex = 0;
+	for (TFieldIterator<FProperty> It(InFunc); It && It->HasAnyPropertyFlags(EPropertyFlags::CPF_Parm); ++It)
+	{
+		const FString PropertyCppType = It->GetCPPType();
+		if (PropertyCppType != ArgsCppTypes[ArgCppTypesIndex])
+		{
+			return false;
+		}
+
+		// Also making sure that the last param is a return param, if we have a return value
+		if constexpr (!std::is_same_v<ReturnType, void>)
+		{
+			if (ArgCppTypesIndex == NumParms - 1 && !It->HasAnyPropertyFlags(EPropertyFlags::CPF_ReturnParm))
+			{
+				return false;
+			}
+		}
+
+		ArgCppTypesIndex++;
+	}
+
+	return true;
+}
+
+template <typename ReturnType, typename... ArgsTypes>
+TValueOrError<ReturnType, void> CallFunc(UObject* InTargetObject, UFunction* InFunc, ArgsTypes&& ...InArgs)
+{
+	if (!InTargetObject || !InFunc)
+	{
+		return MakeError();
+	}
+
+	constexpr int32 NumParms = std::is_same_v <ReturnType, void> ? sizeof...(ArgsTypes) : (sizeof...(ArgsTypes) + 1);
+
+	if (NumParms != InFunc->NumParms)
+	{
+		return MakeError();
+	}
+
+	FFuncParam InParams[NumParms] = { FFuncParam::Make(std::forward<ArgsTypes>(InArgs)) ... };
+
+	auto Invoke = [InTargetObject, InFunc, &InParams, NumParms](ReturnType* OutResult) -> bool
+	{
+		// Validate that the function has a return property if the return type is not void.
+		if (std::is_same_v<ReturnType, void> != (InFunc->GetReturnProperty() == nullptr))
+		{
+			return false;
+		}
+
+		// Allocating our "stack" for the function call on the stack (will be freed when this function is exited)
+		uint8* StackMemory = (uint8*)FMemory_Alloca(InFunc->ParmsSize);
+		FMemory::Memzero(StackMemory, InFunc->ParmsSize);
+
+		if constexpr (!std::is_same_v<ReturnType, void>)
+		{
+			check(OutResult != nullptr);
+			InParams[NumParms - 1] = FFuncParam::Make(*OutResult);
+		}
+
+		bool bValid = true;
+		int32 ParamIndex = 0;
+
+		// Initializing our "stack" with our parameters. Use the property system to make sure more complex types
+		// are constructed before being set.
+		for (TFieldIterator<FProperty> It(InFunc); It && It->HasAnyPropertyFlags(CPF_Parm); ++It)
+		{
+			FProperty* LocalProp = *It;
+			if (!LocalProp->HasAnyPropertyFlags(CPF_ZeroConstructor))
+			{
+				LocalProp->InitializeValue_InContainer(StackMemory);
+			}
+
+			if (bValid)
+			{
+				if (ParamIndex >= NumParms)
+				{
+					bValid = false;
+					continue;
+				}
+
+				FFuncParam& Param = InParams[ParamIndex++];
+
+				if (LocalProp->GetCPPType() != Param.CppType)
+				{
+					bValid = false;
+					continue;
+				}
+
+				LocalProp->SetValue_InContainer(StackMemory, Param.Value);
+			}
+		}
+
+		if (bValid)
+		{
+			FFrame Stack(InTargetObject, InFunc, StackMemory, nullptr, InFunc->ChildProperties);
+			InFunc->Invoke(InTargetObject, Stack, OutResult);
+		}
+
+		ParamIndex = 0;
+		// Copy back all non-const out params (that is not the return param, this one is already set by the invoke call)
+		// from the stack, also making sure that the constructed types are destroyed accordingly.
+		for (TFieldIterator<FProperty> It(InFunc); It && It->HasAnyPropertyFlags(CPF_Parm); ++It)
+		{
+			FProperty* LocalProp = *It;
+			FFuncParam& Param = InParams[ParamIndex++];
+
+			if (LocalProp->HasAnyPropertyFlags(CPF_OutParm) && !LocalProp->HasAnyPropertyFlags(CPF_ReturnParm) && !Param.bIsConst)
+			{
+				LocalProp->GetValue_InContainer(StackMemory, Param.Value);
+			}
+
+			LocalProp->DestroyValue_InContainer(StackMemory);
+		}
+
+		return bValid;
+	};
+
+	if constexpr (std::is_same_v<ReturnType, void>)
+	{
+		if (Invoke(nullptr))
+		{
+			return MakeValue();
+		}
+		else
+		{
+			return MakeError();
+		}
+	}
+	else
+	{
+		ReturnType OutResult{};
+		if (Invoke(&OutResult))
+		{
+			return MakeValue(std::move(OutResult));
+		}
+		else
+		{
+			return MakeError();
+		}
+	}
+}
+
 } // UE::StructUtils::Private
 
 
@@ -598,17 +860,42 @@ TSharedRef<SWidget> FPropertyBagInstanceDataDetails::OnPropertyNameContent(TShar
 	constexpr bool bInShouldCloseWindowAfterMenuSelection = true;
 	FMenuBuilder MenuBuilder(bInShouldCloseWindowAfterMenuSelection, nullptr);
 
-	auto GetFilteredVariableTypeTree = [](TArray<TSharedPtr<UEdGraphSchema_K2::FPinTypeTreeInfo>>& TypeTree, ETypeTreeFilter TypeTreeFilter)
+	auto GetFilteredVariableTypeTree = [BagStructProperty = BagStructProperty](TArray<TSharedPtr<UEdGraphSchema_K2::FPinTypeTreeInfo>>& TypeTree, ETypeTreeFilter TypeTreeFilter)
 	{
+		UFunction* IsPinTypeAcceptedFunc = nullptr;
+		UObject* IsPinTypeAcceptedTarget = nullptr;
+		UE::StructUtils::Private::FindUserFunction(BagStructProperty, UE::StructUtils::Metadata::IsPinTypeAcceptedName ,IsPinTypeAcceptedFunc, IsPinTypeAcceptedTarget);
+
+		// We need to make sure the signature matches perfectly: bool(FEdGraphPinType)
+		bool bFuncIsValid = UE::StructUtils::Private::ValidateFunctionSignature<bool, FEdGraphPinType>(IsPinTypeAcceptedFunc);
+		if (!ensureMsgf(bFuncIsValid, TEXT("[%s] Function %s does not have the right signature."), *UE::StructUtils::Metadata::IsPinTypeAcceptedName.ToString(), *IsPinTypeAcceptedFunc->GetName()))
+		{
+			return;
+		}
+
+		auto IsPinTypeAccepted = [IsPinTypeAcceptedFunc, IsPinTypeAcceptedTarget](const FEdGraphPinType& InPinType) -> bool
+		{
+			if (IsPinTypeAcceptedFunc && IsPinTypeAcceptedTarget)
+			{
+				const TValueOrError<bool, void> bIsValid = UE::StructUtils::Private::CallFunc<bool>(IsPinTypeAcceptedTarget, IsPinTypeAcceptedFunc, InPinType);
+				return bIsValid.HasValue() && bIsValid.GetValue();
+			}
+			else
+			{
+				return true;
+			}
+		};
+
 		check(GetDefault<UEdGraphSchema_K2>());
-		GetDefault<UPropertyBagSchema>()->GetVariableTypeTree(TypeTree, TypeTreeFilter);
+		TArray<TSharedPtr<UEdGraphSchema_K2::FPinTypeTreeInfo>> TempTypeTree;
+		GetDefault<UPropertyBagSchema>()->GetVariableTypeTree(TempTypeTree, TypeTreeFilter);
 
 		// Filter
-		for (TSharedPtr<UEdGraphSchema_K2::FPinTypeTreeInfo>& PinType : TypeTree)
+		for (TSharedPtr<UEdGraphSchema_K2::FPinTypeTreeInfo>& PinType : TempTypeTree)
 		{
-			if (!PinType.IsValid())
+			if (!PinType.IsValid() || !IsPinTypeAccepted(PinType->GetPinType(/*bForceLoadSubCategoryObject*/false)))
 			{
-				return;
+				continue;
 			}
 
 			for (int32 ChildIndex = 0; ChildIndex < PinType->Children.Num(); )
@@ -616,7 +903,9 @@ TSharedRef<SWidget> FPropertyBagInstanceDataDetails::OnPropertyNameContent(TShar
 				TSharedPtr<UEdGraphSchema_K2::FPinTypeTreeInfo> Child = PinType->Children[ChildIndex];
 				if (Child.IsValid())
 				{
-					if (!UE::StructUtils::Private::CanHaveMemberVariableOfType(Child->GetPinType(/*bForceLoadSubCategoryObject*/false)))
+					const FEdGraphPinType& ChildPinType = Child->GetPinType(/*bForceLoadSubCategoryObject*/false);
+
+					if (!UE::StructUtils::Private::CanHaveMemberVariableOfType(ChildPinType) || !IsPinTypeAccepted(ChildPinType))
 					{
 						PinType->Children.RemoveAt(ChildIndex);
 						continue;
@@ -624,7 +913,9 @@ TSharedRef<SWidget> FPropertyBagInstanceDataDetails::OnPropertyNameContent(TShar
 				}
 				++ChildIndex;
 			}
-		}			
+
+			TypeTree.Add(PinType);
+		}
 	};
 
 	auto GetPinInfo = [BagStructProperty = BagStructProperty, ChildPropertyHandle]()
@@ -674,6 +965,36 @@ TSharedRef<SWidget> FPropertyBagInstanceDataDetails::OnPropertyNameContent(TShar
 			return;
 		}
 
+		// Extra check provided by the user to cancel a remove action. Useful to provide the user a possibility to cancel the action if
+		// the given property is in use elsewhere.
+		UFunction* CanRemovePropertyFunc = nullptr;
+		UObject* CanRemovePropertyTarget = nullptr;
+		if (UE::StructUtils::Private::FindUserFunction(BagStructProperty, UE::StructUtils::Metadata::CanRemovePropertyName, CanRemovePropertyFunc, CanRemovePropertyTarget))
+		{
+			FName PropertyName = ChildPropertyHandle->GetProperty()->GetFName();
+			const UPropertyBag* PropertyBag = UE::StructUtils::Private::GetCommonBagStruct(BagStructProperty);
+			const FPropertyBagPropertyDesc* PropertyDesc = PropertyBag ? PropertyBag->FindPropertyDescByName(PropertyName) : nullptr;
+
+			if (!PropertyDesc)
+			{
+				return;
+			}
+
+			// We need to make sure the signature matches perfectly: bool(FGuid, FName)s
+			bool bFuncIsValid = UE::StructUtils::Private::ValidateFunctionSignature<bool, FGuid, FName>(CanRemovePropertyFunc);
+			if(!ensureMsgf(bFuncIsValid, TEXT("[%s] Function %s does not have the right signature."), *UE::StructUtils::Metadata::CanRemovePropertyName.ToString(), *CanRemovePropertyFunc->GetName()))
+			{
+				return;
+			}
+
+			const TValueOrError<bool, void> bCanRemove = UE::StructUtils::Private::CallFunc<bool>(CanRemovePropertyTarget, CanRemovePropertyFunc, PropertyDesc->ID, PropertyDesc->Name);
+
+			if (bCanRemove.HasError() || !bCanRemove.GetValue())
+			{
+				return;
+			}
+		}
+
 		UE::StructUtils::Private::ApplyChangesToPropertyDescs(
 			LOCTEXT("OnPropertyRemoved", "Remove Property"), BagStructProperty, PropUtils,
 			[&ChildPropertyHandle](TArray<FPropertyBagPropertyDesc>& PropertyDescs)
@@ -883,6 +1204,4 @@ bool UPropertyBagSchema::SupportsPinTypeContainer(TWeakPtr<const FEdGraphSchemaA
 	return ContainerType == EPinContainerType::None || ContainerType == EPinContainerType::Array;
 }
 
-
 #undef LOCTEXT_NAMESPACE
-
