[StructUtils] Fixing crash in PropertyBagDetails: Dereferencing nullptr

#rnx

diff --git a/Engine/Plugins/Experimental/StructUtils/Source/StructUtilsEditor/Private/Customizations/PropertyBagDetails.cpp b/Engine/Plugins/Experimental/StructUtils/Source/StructUtilsEditor/Private/Customizations/PropertyBagDetails.cpp
--- a/Engine/Plugins/Experimental/StructUtils/Source/StructUtilsEditor/Private/Customizations/PropertyBagDetails.cpp
+++ b/Engine/Plugins/Experimental/StructUtils/Source/StructUtilsEditor/Private/Customizations/PropertyBagDetails.cpp
@@ -864,13 +864,16 @@ TSharedRef<SWidget> FPropertyBagInstanceDataDetails::OnPropertyNameContent(TShar
 	{
 		UFunction* IsPinTypeAcceptedFunc = nullptr;
 		UObject* IsPinTypeAcceptedTarget = nullptr;
-		UE::StructUtils::Private::FindUserFunction(BagStructProperty, UE::StructUtils::Metadata::IsPinTypeAcceptedName ,IsPinTypeAcceptedFunc, IsPinTypeAcceptedTarget);
-
-		// We need to make sure the signature matches perfectly: bool(FEdGraphPinType)
-		bool bFuncIsValid = UE::StructUtils::Private::ValidateFunctionSignature<bool, FEdGraphPinType>(IsPinTypeAcceptedFunc);
-		if (!ensureMsgf(bFuncIsValid, TEXT("[%s] Function %s does not have the right signature."), *UE::StructUtils::Metadata::IsPinTypeAcceptedName.ToString(), *IsPinTypeAcceptedFunc->GetName()))
+		if (UE::StructUtils::Private::FindUserFunction(BagStructProperty, UE::StructUtils::Metadata::IsPinTypeAcceptedName, IsPinTypeAcceptedFunc, IsPinTypeAcceptedTarget))
 		{
-			return;
+			check(IsPinTypeAcceptedFunc && IsPinTypeAcceptedTarget);
+
+			// We need to make sure the signature matches perfectly: bool(FEdGraphPinType)
+			bool bFuncIsValid = UE::StructUtils::Private::ValidateFunctionSignature<bool, FEdGraphPinType>(IsPinTypeAcceptedFunc);
+			if (!ensureMsgf(bFuncIsValid, TEXT("[%s] Function %s does not have the right signature."), *UE::StructUtils::Metadata::IsPinTypeAcceptedName.ToString(), *IsPinTypeAcceptedFunc->GetName()))
+			{
+				return;
+			}
 		}
 
 		auto IsPinTypeAccepted = [IsPinTypeAcceptedFunc, IsPinTypeAcceptedTarget](const FEdGraphPinType& InPinType) -> bool
@@ -971,6 +974,8 @@ TSharedRef<SWidget> FPropertyBagInstanceDataDetails::OnPropertyNameContent(TShar
 		UObject* CanRemovePropertyTarget = nullptr;
 		if (UE::StructUtils::Private::FindUserFunction(BagStructProperty, UE::StructUtils::Metadata::CanRemovePropertyName, CanRemovePropertyFunc, CanRemovePropertyTarget))
 		{
+			check(CanRemovePropertyFunc && CanRemovePropertyTarget);
+
 			FName PropertyName = ChildPropertyHandle->GetProperty()->GetFName();
 			const UPropertyBag* PropertyBag = UE::StructUtils::Private::GetCommonBagStruct(BagStructProperty);
 			const FPropertyBagPropertyDesc* PropertyDesc = PropertyBag ? PropertyBag->FindPropertyDescByName(PropertyName) : nullptr;
