Property Bag Test - use less aggressive GC

diff --git a/Engine/Plugins/Experimental/StructUtils/Source/StructUtilsTestSuite/Private/PropertyBagTest.cpp b/Engine/Plugins/Experimental/StructUtils/Source/StructUtilsTestSuite/Private/PropertyBagTest.cpp
--- a/Engine/Plugins/Experimental/StructUtils/Source/StructUtilsTestSuite/Private/PropertyBagTest.cpp
+++ b/Engine/Plugins/Experimental/StructUtils/Source/StructUtilsTestSuite/Private/PropertyBagTest.cpp
@@ -408,17 +408,18 @@ struct FTest_GC : FAITestBase
 
 		// Obj is unreachable, it should be collected by the GC.
 		Obj = nullptr;
-		CollectGarbage(RF_NoFlags);
+		CollectGarbage(GARBAGE_COLLECTION_KEEPFLAGS);
 
 		// The used property bag struct should exists after the GC.
 		const UPropertyBag* ExistingObj = FindObject<UPropertyBag>(GetTransientPackage(), *ObjName);
-		const UPropertyBag* ExistingBagStruct1 = FindObject<UPropertyBag>(GetTransientPackage(), *BagStructName);
+		UPropertyBag* ExistingBagStruct1 = FindObject<UPropertyBag>(GetTransientPackage(), *BagStructName);
 
 		AITEST_NULL(TEXT("Obj should have been released"), ExistingObj);
 		AITEST_NOT_NULL(TEXT("Bag struct should exists after Obj released"), ExistingBagStruct1);
 
 		// The next GC should collect the bag struct
-		CollectGarbage(RF_NoFlags);
+		ExistingBagStruct1->ClearFlags(RF_Standalone);
+		CollectGarbage(GARBAGE_COLLECTION_KEEPFLAGS);
 
 		const UPropertyBag* ExistingBagStruct2 = FindObject<UPropertyBag>(GetTransientPackage(), *BagStructName);
 		AITEST_NULL(TEXT("Bag struct should not exists after second GC"), ExistingBagStruct2);
