Improved multi-edit behavior of instanced struct instances

Editing multiple instanced struct instances will now edit the properties of the common struct of all instances (like object editing), rather than only editing instanced struct instances using the same struct type as the first instance being edited.

diff --git a/Engine/Plugins/Experimental/StructUtils/Source/StructUtilsEditor/Private/Customizations/InstancedStructDetails.cpp b/Engine/Plugins/Experimental/StructUtils/Source/StructUtilsEditor/Private/Customizations/InstancedStructDetails.cpp
--- a/Engine/Plugins/Experimental/StructUtils/Source/StructUtilsEditor/Private/Customizations/InstancedStructDetails.cpp
+++ b/Engine/Plugins/Experimental/StructUtils/Source/StructUtilsEditor/Private/Customizations/InstancedStructDetails.cpp
@@ -60,13 +60,23 @@ public:
 	
 	virtual const UStruct* GetBaseStructure() const override
 	{
+		// Taken from UClass::FindCommonBase
+		auto FindCommonBaseStruct = [](const UScriptStruct* StructA, const UScriptStruct* StructB)
+		{
+			const UScriptStruct* CommonBaseStruct = StructA;
+			while (CommonBaseStruct && StructB && !StructB->IsChildOf(CommonBaseStruct))
+			{
+				CommonBaseStruct = Cast<UScriptStruct>(CommonBaseStruct->GetSuperStruct());
+			}
+			return CommonBaseStruct;
+		};
+
 		const UScriptStruct* CommonStruct = nullptr;
-		EnumerateInstances([&CommonStruct](const UScriptStruct* ScriptStruct, uint8* Memory, UPackage* Package)
+		EnumerateInstances([&CommonStruct, &FindCommonBaseStruct](const UScriptStruct* ScriptStruct, uint8* Memory, UPackage* Package)
 		{
-			if (!CommonStruct && ScriptStruct)
+			if (ScriptStruct)
 			{
-				CommonStruct = ScriptStruct;
-				return false; // Stop
+				CommonStruct = FindCommonBaseStruct(ScriptStruct, CommonStruct);
 			}
 			return true; // Continue
 		});
@@ -78,17 +88,15 @@ public:
 	{
 		// The returned instances need to be compatible with base structure.
 		// This function returns empty instances in case they are not compatible, with the idea that we have as many instances as we have outer objects.
-		const UScriptStruct* CommonStruct = nullptr;
-		EnumerateInstances([&OutInstances, &CommonStruct](const UScriptStruct* ScriptStruct, uint8* Memory, UPackage* Package)
+		const UScriptStruct* CommonStruct = Cast<UScriptStruct>(GetBaseStructure());
+		EnumerateInstances([&OutInstances, CommonStruct](const UScriptStruct* ScriptStruct, uint8* Memory, UPackage* Package)
 		{
 			TSharedPtr<FStructOnScope> Result;
 			
-			// Only return structs that are same type.
-			if (CommonStruct == nullptr || ScriptStruct == CommonStruct)
+			if (CommonStruct && ScriptStruct && ScriptStruct->IsChildOf(CommonStruct))
 			{
 				Result = MakeShared<FStructOnScope>(ScriptStruct, Memory);
 				Result->SetPackage(Package);
-				CommonStruct = ScriptStruct;
 			}
 
 			OutInstances.Add(Result);
@@ -110,12 +118,12 @@ public:
 		}
 
 		FInstancedStruct& InstancedStruct = *reinterpret_cast<FInstancedStruct*>(ParentValueAddress);
-		if (InstancedStruct.GetScriptStruct() != ExpectedType)
+		if (ExpectedType && InstancedStruct.GetScriptStruct() && InstancedStruct.GetScriptStruct()->IsChildOf(ExpectedType))
 		{
-			return nullptr;
+			return InstancedStruct.GetMutableMemory();
 		}
 
-		return InstancedStruct.GetMutableMemory();
+		return nullptr;
 	}
 	
 protected:
