Allow AutoRTFM to work on StructUtils.

diff --git a/Engine/Plugins/Experimental/StructUtils/Source/StructUtils/StructUtils.Build.cs b/Engine/Plugins/Experimental/StructUtils/Source/StructUtils/StructUtils.Build.cs
--- a/Engine/Plugins/Experimental/StructUtils/Source/StructUtils/StructUtils.Build.cs
+++ b/Engine/Plugins/Experimental/StructUtils/Source/StructUtils/StructUtils.Build.cs
@@ -21,6 +21,8 @@ namespace UnrealBuildTool.Rules
 			{
 				PrivateDependencyModuleNames.Add("Engine");
 			}
+
+			bAllowAutoRTFMInstrumentation = true;
 		}
 	}
 }
