StructUtils: Added some FInstancedStruct tests.

diff --git a/Engine/Plugins/Experimental/StructUtils/Source/StructUtilsTestSuite/Private/InstancedStructTest.cpp b/Engine/Plugins/Experimental/StructUtils/Source/StructUtilsTestSuite/Private/InstancedStructTest.cpp
new file mode 100644
--- /dev/null
+++ b/Engine/Plugins/Experimental/StructUtils/Source/StructUtilsTestSuite/Private/InstancedStructTest.cpp
@@ -0,0 +1,118 @@
+﻿// Copyright Epic Games, Inc. All Rights Reserved.
+
+#include "AITestsCommon.h"
+#include "SharedStruct.h"
+#include "StructUtilsTestTypes.h"
+
+#define LOCTEXT_NAMESPACE "StructUtilsTests"
+
+UE_DISABLE_OPTIMIZATION_SHIP
+
+namespace FInstancedStructTest
+{
+	struct FTest_InstancedStructCreate : FAITestBase
+	{
+		virtual bool InstantTest() override
+		{
+			constexpr float Val = 99.f;
+			
+			{
+				FInstancedStruct InstancedStruct = FInstancedStruct::Make<FTestStructSimpleNonZeroDefault>();
+
+				AITEST_EQUAL("FInstancedStruct default initialized from Make should have same value as default constructed", FTestStructSimpleNonZeroDefault(), InstancedStruct.Get<FTestStructSimpleNonZeroDefault>());
+			}
+
+			{
+				FTestStructSimple Simple(Val);
+
+				FInstancedStruct InstancedStruct = FInstancedStruct::Make(Simple);
+
+				AITEST_EQUAL("FInstancedStruct initialized from Make should have value of FTestStructSimple its initiliazed from", Val, InstancedStruct.Get<FTestStructSimple>().Float);
+			}
+
+			{
+				FTestStructSimple Simple(Val);
+				FStructView StructView = FStructView::Make(Simple);
+				FInstancedStruct InstancedStruct(StructView);
+
+				AITEST_EQUAL("FInstancedStruct initialized from Make should have value of FStructView its initiliazed from", Val, InstancedStruct.Get<FTestStructSimple>().Float);
+			}
+
+			{
+				FTestStructSimple Simple(Val);
+
+				FInstancedStruct InstancedStruct = FInstancedStruct::Make<FTestStructSimple>(Val);
+
+				AITEST_EQUAL("FInstancedStruct initialized from Make should have value reflecting TArgs", Val, InstancedStruct.Get<FTestStructSimple>().Float);
+			}
+
+			{
+				FTestStructSimple Simple(Val);
+
+				FInstancedStruct InstancedStruct;
+				InstancedStruct.InitializeAs<FTestStructSimple>(Val);
+
+				AITEST_EQUAL("FInstancedStruct initialized from Make should have value reflecting TArgs", Val, InstancedStruct.Get<FTestStructSimple>().Float);
+			}
+
+			{
+				FTestStructSimple Simple(Val);
+
+				FInstancedStruct InstancedStruct;
+				InstancedStruct.InitializeAs<FTestStructSimple>(Val);
+				AITEST_EQUAL("FInstancedStruct initialized from InitializeAs should have value reflecting TArgs", Val, InstancedStruct.Get<FTestStructSimple>().Float);
+
+				InstancedStruct.InitializeAs<FTestStructSimpleNonZeroDefault>();
+				AITEST_EQUAL("FInstancedStruct initialized from InitializeAs should have same value as default constructed", FTestStructSimpleNonZeroDefault(), InstancedStruct.Get<FTestStructSimpleNonZeroDefault>());
+
+				InstancedStruct.InitializeAs(nullptr);
+				AITEST_FALSE("FInstancedStruct initialized from InitializeAs with empty struct should not be valid", InstancedStruct.IsValid());
+			}
+
+			return true;
+		}
+	};
+
+	IMPLEMENT_AI_INSTANT_TEST(FTest_InstancedStructCreate, "System.StructUtils.InstancedStruct.Make");
+
+	struct FTest_InstancedStructBasic : FAITestBase
+	{
+		virtual bool InstantTest() override
+		{
+			{
+				FInstancedStruct InstancedStruct = FInstancedStruct::Make<FTestStructSimple>();
+				FInstancedStruct InstancedStruct2(InstancedStruct);
+
+				AITEST_EQUAL("InstancedStruct and InstancedStruct2 should be equal from copy construction", InstancedStruct, InstancedStruct2);
+			}
+
+			{
+				FInstancedStruct InstancedStruct = FInstancedStruct::Make<FTestStructSimple>();
+				FInstancedStruct InstancedStruct2;
+
+				InstancedStruct2 = InstancedStruct;
+
+				AITEST_EQUAL("FInstancedStruct and FInstancedStruct should be equal from copy assignment", InstancedStruct, InstancedStruct2);
+			}
+
+			{
+				FInstancedStruct InstancedStruct;
+				AITEST_FALSE("Default constructed FInstancedStruct should IsValid() == false", InstancedStruct.IsValid());
+			}
+
+			{
+				FInstancedStruct InstancedStruct = FInstancedStruct::Make<FTestStructSimple>();
+				AITEST_TRUE("FInstancedStruct created to a specific struct type should be IsValid()", InstancedStruct.IsValid());
+			}
+
+			return true;
+		}
+	};
+	IMPLEMENT_AI_INSTANT_TEST(FTest_InstancedStructBasic, "System.StructUtils.InstancedStruct.Basic");
+
+}
+
+
+UE_ENABLE_OPTIMIZATION_SHIP
+
+#undef LOCTEXT_NAMESPACE
