StructUtils: Change FPropertyBagContainerTypes to use GetArrayHash()

diff --git a/Engine/Plugins/Experimental/StructUtils/Source/StructUtils/Public/PropertyBag.h b/Engine/Plugins/Experimental/StructUtils/Source/StructUtils/Public/PropertyBag.h
--- a/Engine/Plugins/Experimental/StructUtils/Source/StructUtils/Public/PropertyBag.h
+++ b/Engine/Plugins/Experimental/StructUtils/Source/StructUtils/Public/PropertyBag.h
@@ -138,12 +138,7 @@ struct STRUCTUTILS_API FPropertyBagContainerTypes
 
 	friend FORCEINLINE uint32 GetTypeHash(const FPropertyBagContainerTypes& PropertyBagContainerTypes)
 	{
-		uint32 Hash = GetTypeHash(PropertyBagContainerTypes.NumContainers);
-		for (const EPropertyBagContainerType Type : PropertyBagContainerTypes.Types)
-		{
-			Hash = HashCombine(Hash, GetTypeHash(Type));
-		}
-		return Hash;
+		return GetArrayHash(PropertyBagContainerTypes.Types.GetData(), PropertyBagContainerTypes.NumContainers);
 	}
 
 	EPropertyBagContainerType* begin() { return &Types[0]; }
