[StructUtils] Add SoftPath getter/setter support in PropertyBag to interact with properties that are SoftObjectPtr or SoftObjectClass.

* Otherwise we need to redo logic on the client side to set soft object pointers, when objects are not yet loaded.

diff --git a/Engine/Plugins/Experimental/StructUtils/Source/StructUtils/Private/PropertyBag.cpp b/Engine/Plugins/Experimental/StructUtils/Source/StructUtils/Private/PropertyBag.cpp
--- a/Engine/Plugins/Experimental/StructUtils/Source/StructUtils/Private/PropertyBag.cpp
+++ b/Engine/Plugins/Experimental/StructUtils/Source/StructUtils/Private/PropertyBag.cpp
@@ -661,6 +661,34 @@ namespace UE::StructUtils::Private
 		return EPropertyBagResult::Success;
 	}
 
+	EPropertyBagResult GetPropertyValueAsSoftPath(const FPropertyBagPropertyDesc* Desc, const void* Address, FSoftObjectPath& OutValue)
+	{
+		if (Desc == nullptr || Desc->CachedProperty == nullptr)
+		{
+			return EPropertyBagResult::PropertyNotFound;
+		}
+		if (Desc->ValueType != EPropertyBagPropertyType::SoftObject
+			&& Desc->ValueType != EPropertyBagPropertyType::SoftClass)
+		{
+			return EPropertyBagResult::TypeMismatch;
+		}
+		if (Address == nullptr)
+		{
+			return EPropertyBagResult::OutOfBounds;
+		}
+		if (Desc->ContainerTypes.Num() > 0)
+		{
+			return EPropertyBagResult::TypeMismatch;
+		}
+
+		const FSoftObjectProperty* SoftObjectProperty = CastFieldChecked<FSoftObjectProperty>(Desc->CachedProperty);
+		check(SoftObjectProperty->PropertyClass);
+
+		OutValue = SoftObjectProperty->GetPropertyValue(Address).ToSoftObjectPath();
+
+		return EPropertyBagResult::Success;
+	}
+
 	//----------------------------------------------------------------//
 	//  Setters
 	//----------------------------------------------------------------//
@@ -955,6 +983,35 @@ namespace UE::StructUtils::Private
 		return EPropertyBagResult::Success;
 	}
 
+	EPropertyBagResult SetPropertyValueAsSoftPath(const FPropertyBagPropertyDesc* Desc, void* Address, const FSoftObjectPath& InValue)
+	{
+		if (Desc == nullptr || Desc->CachedProperty == nullptr)
+		{
+			return EPropertyBagResult::PropertyNotFound;
+		}
+		if (Desc->ValueType != EPropertyBagPropertyType::SoftObject &&
+			Desc->ValueType != EPropertyBagPropertyType::SoftClass)
+		{
+			return EPropertyBagResult::TypeMismatch;
+		}
+		if (Address == nullptr)
+		{
+			return EPropertyBagResult::OutOfBounds;
+		}
+		if (Desc->ContainerTypes.Num() > 0)
+		{
+			return EPropertyBagResult::TypeMismatch;
+		}
+
+		const FSoftObjectProperty* SoftObjectProperty = CastFieldChecked<FSoftObjectProperty>(Desc->CachedProperty);
+		check(SoftObjectProperty->PropertyClass);
+		check(Desc->ValueTypeObject);
+
+		SoftObjectProperty->SetPropertyValue(Address, FSoftObjectPtr(InValue));
+
+		return EPropertyBagResult::Success;
+	}
+
 	void CopyMatchingValuesByID(const FConstStructView Source, FStructView Target)
 	{
 		if (!Source.IsValid() || !Target.IsValid())
@@ -1559,6 +1616,19 @@ TValueOrError<UClass*, EPropertyBagResult> FInstancedPropertyBag::GetValueClass(
 	return MakeValue(Class);
 }
 
+TValueOrError<FSoftObjectPath, EPropertyBagResult> FInstancedPropertyBag::GetValueSoftPath(const FName Name) const
+{
+	const FPropertyBagPropertyDesc* Desc = FindPropertyDescByName(Name);
+	FSoftObjectPath ReturnValue;
+	const EPropertyBagResult Result = UE::StructUtils::Private::GetPropertyValueAsSoftPath(Desc, GetValueAddress(Desc), ReturnValue);
+	if (Result != EPropertyBagResult::Success)
+	{
+		return MakeError(Result);
+	}
+
+	return MakeValue(std::move(ReturnValue));
+}
+
 TValueOrError<FString, EPropertyBagResult> FInstancedPropertyBag::GetValueSerializedString(const FName Name)
 {
 	const FPropertyBagPropertyDesc* Desc = FindPropertyDescByName(Name);
@@ -1659,6 +1729,12 @@ EPropertyBagResult FInstancedPropertyBag::SetValueClass(const FName Name, UClass
 	return UE::StructUtils::Private::SetPropertyValueAsObject(Desc, GetMutableValueAddress(Desc), InValue);
 }
 
+EPropertyBagResult FInstancedPropertyBag::SetValueSoftPath(const FName Name, const FSoftObjectPath& InValue)
+{
+	const FPropertyBagPropertyDesc* Desc = FindPropertyDescByName(Name);
+	return UE::StructUtils::Private::SetPropertyValueAsSoftPath(Desc, GetMutableValueAddress(Desc), InValue);
+}
+
 EPropertyBagResult FInstancedPropertyBag::SetValueSerializedString(const FName Name, const FString& InValue)
 {
 	const FPropertyBagPropertyDesc* Desc = FindPropertyDescByName(Name);
@@ -2113,6 +2189,17 @@ TValueOrError<const FPropertyBagArrayRef, EPropertyBagResult> FPropertyBagArrayR
 	return MakeValue(FPropertyBagArrayRef(ValueDesc, Address));
 }
 
+TValueOrError<FSoftObjectPath, EPropertyBagResult> FPropertyBagArrayRef::GetValueSoftPath(const int32 Index) const
+{
+	FSoftObjectPath ReturnValue;
+	const EPropertyBagResult Result = UE::StructUtils::Private::GetPropertyValueAsSoftPath(&ValueDesc, GetAddress(Index), ReturnValue);
+	if (Result != EPropertyBagResult::Success)
+	{
+		return MakeError(Result);
+	}
+	return MakeValue(std::move(ReturnValue));
+}
+
 EPropertyBagResult FPropertyBagArrayRef::SetValueBool(const int32 Index, const bool bInValue)
 {
 	return UE::StructUtils::Private::SetPropertyFromInt64(&ValueDesc, GetMutableAddress(Index), bInValue ? 1 : 0);
@@ -2178,6 +2265,10 @@ EPropertyBagResult FPropertyBagArrayRef::SetValueClass(const int32 Index, UClass
 	return UE::StructUtils::Private::SetPropertyValueAsObject(&ValueDesc, GetMutableAddress(Index), InValue);
 }
 
+EPropertyBagResult FPropertyBagArrayRef::SetValueSoftPath(const int32 Index, const FSoftObjectPath& InValue)
+{
+	return UE::StructUtils::Private::SetPropertyValueAsSoftPath(&ValueDesc, GetMutableAddress(Index), InValue);
+}
 
 //----------------------------------------------------------------//
 //  UPropertyBag
diff --git a/Engine/Plugins/Experimental/StructUtils/Source/StructUtils/Public/PropertyBag.h b/Engine/Plugins/Experimental/StructUtils/Source/StructUtils/Public/PropertyBag.h
--- a/Engine/Plugins/Experimental/StructUtils/Source/StructUtils/Public/PropertyBag.h
+++ b/Engine/Plugins/Experimental/StructUtils/Source/StructUtils/Public/PropertyBag.h
@@ -464,6 +464,7 @@ struct STRUCTUTILS_API FInstancedPropertyBag
 	TValueOrError<FStructView, EPropertyBagResult> GetValueStruct(const FName Name, const UScriptStruct* RequestedStruct = nullptr) const;
 	TValueOrError<UObject*, EPropertyBagResult> GetValueObject(const FName Name, const UClass* RequestedClass = nullptr) const;
 	TValueOrError<UClass*, EPropertyBagResult> GetValueClass(const FName Name) const;
+	TValueOrError<FSoftObjectPath, EPropertyBagResult> GetValueSoftPath(const FName Name) const;
 
 	/** @return string-based serialized representation of the value. */
 	TValueOrError<FString, EPropertyBagResult> GetValueSerializedString(const FName Name);
@@ -537,6 +538,7 @@ struct STRUCTUTILS_API FInstancedPropertyBag
 	EPropertyBagResult SetValueStruct(const FName Name, FConstStructView InValue);
 	EPropertyBagResult SetValueObject(const FName Name, UObject* InValue);
 	EPropertyBagResult SetValueClass(const FName Name, UClass* InValue);
+	EPropertyBagResult SetValueSoftPath(const FName Name, const FSoftObjectPath& InValue);
 
 	/**
 	 * Sets property value from a serialized representation of the value. If the string value provided
@@ -685,6 +687,7 @@ public:
 	TValueOrError<FStructView, EPropertyBagResult> GetValueStruct(const int32 Index, const UScriptStruct* RequestedStruct = nullptr) const;
 	TValueOrError<UObject*, EPropertyBagResult> GetValueObject(const int32 Index, const UClass* RequestedClass = nullptr) const;
 	TValueOrError<UClass*, EPropertyBagResult> GetValueClass(const int32 Index) const;
+	TValueOrError<FSoftObjectPath, EPropertyBagResult> GetValueSoftPath(const int32 Index) const;
 
 	/** @return enum value of specified type. */
 	template <typename T>
@@ -769,6 +772,7 @@ public:
 	EPropertyBagResult SetValueStruct(const int32 Index, FConstStructView InValue);
 	EPropertyBagResult SetValueObject(const int32 Index, UObject* InValue);
 	EPropertyBagResult SetValueClass(const int32 Index, UClass* InValue);
+	EPropertyBagResult SetValueSoftPath(const int32 Index, const FSoftObjectPath& InValue);
 
 	/** Sets enum value specified type. */
 	template <typename T>
