PropertyBag: Potential fix for crash on max/linux

diff --git a/Engine/Plugins/Experimental/StructUtils/Source/StructUtilsEditor/Private/Customizations/PropertyBagDetails.cpp b/Engine/Plugins/Experimental/StructUtils/Source/StructUtilsEditor/Private/Customizations/PropertyBagDetails.cpp
--- a/Engine/Plugins/Experimental/StructUtils/Source/StructUtilsEditor/Private/Customizations/PropertyBagDetails.cpp
+++ b/Engine/Plugins/Experimental/StructUtils/Source/StructUtilsEditor/Private/Customizations/PropertyBagDetails.cpp
@@ -862,6 +862,12 @@ TSharedRef<SWidget> FPropertyBagInstanceDataDetails::OnPropertyNameContent(TShar
 
 	auto GetFilteredVariableTypeTree = [BagStructProperty = BagStructProperty](TArray<TSharedPtr<UEdGraphSchema_K2::FPinTypeTreeInfo>>& TypeTree, ETypeTreeFilter TypeTreeFilter)
 	{
+		// The SPinTypeSelector popup might outlive this details view, so bag struct property can be invalid here.
+    	if (!BagStructProperty || !BagStructProperty->IsValidHandle())
+    	{
+    		return;
+    	}
+
 		UFunction* IsPinTypeAcceptedFunc = nullptr;
 		UObject* IsPinTypeAcceptedTarget = nullptr;
 		if (UE::StructUtils::Private::FindUserFunction(BagStructProperty, UE::StructUtils::Metadata::IsPinTypeAcceptedName, IsPinTypeAcceptedFunc, IsPinTypeAcceptedTarget))
