Fixed minor typo in FConstStructView description

diff --git a/Engine/Plugins/Experimental/StructUtils/Source/StructUtils/Public/StructView.h b/Engine/Plugins/Experimental/StructUtils/Source/StructUtils/Public/StructView.h
--- a/Engine/Plugins/Experimental/StructUtils/Source/StructUtils/Public/StructView.h
+++ b/Engine/Plugins/Experimental/StructUtils/Source/StructUtils/Public/StructView.h
@@ -141,7 +141,7 @@ protected:
  * e.g. FConstStructView A; FStructView B = A; // doesn't compile as the struct data for A is immutable (but mutable for B)
  * e.g. FStructView A; A.Foo = NewVal; // compiles as the struct data for A is mutable.
  * e.g. FConstStructView A; A.Foo = NewVal; // doesn't compile as the struct data for A is immutable.
- * e.g. FStructView A; const FStructView B; A = B; // compiles as the struct B is pointing to can't be made to point at something else but A isn't cosnt.
+ * e.g. FStructView A; const FStructView B; A = B; // compiles as the struct B is pointing to can't be made to point at something else but A isn't const.
  * e.g. const FStructView A; FStructView B; A = B; // doesn't compile as attempting to make const view point at something else
  */
 struct FConstStructView
