[FInstancedStruct] replaced versioning implementation to use archive custom versioning.
Also make serialization symmetrical so struct serialized in cooked builds can be deserialized in Editor builds.

diff --git a/Engine/Plugins/Experimental/StructUtils/Source/StructUtils/Private/InstancedStruct.cpp b/Engine/Plugins/Experimental/StructUtils/Source/StructUtils/Private/InstancedStruct.cpp
--- a/Engine/Plugins/Experimental/StructUtils/Source/StructUtils/Private/InstancedStruct.cpp
+++ b/Engine/Plugins/Experimental/StructUtils/Source/StructUtils/Private/InstancedStruct.cpp
@@ -17,6 +17,30 @@
 
 #include UE_INLINE_GENERATED_CPP_BY_NAME(InstancedStruct)
 
+
+struct STRUCTUTILS_API FInstancedStructCustomVersion
+{
+	enum EType
+	{
+		// Before any version changes were made
+		CustomVersionAdded = 0,
+
+		// -----<new versions can be added above this line>-----
+		VersionPlusOne,
+		LatestVersion = VersionPlusOne - 1
+	};
+
+	// The GUID for this custom version number
+	const FGuid GUID;
+	FCustomVersionRegistration Registration;
+
+	FInstancedStructCustomVersion()
+	: GUID{0xE21E1CAA, 0xAF47425E, 0x89BF6AD4, 0x4C44A8BB}
+	, Registration(GUID, FInstancedStructCustomVersion::LatestVersion, TEXT("InstancedStructCustomVersion"))
+	{}
+};
+static FInstancedStructCustomVersion GInstancedStructCustomVersion;
+
 namespace UE::StructUtils::Private
 {
 #if WITH_EDITORONLY_DATA
@@ -129,57 +153,31 @@ void FInstancedStruct::Reset()
 
 bool FInstancedStruct::Serialize(FArchive& Ar)
 {
-	UScriptStruct* NonConstStruct = const_cast<UScriptStruct*>(GetScriptStruct());
-
-	enum class EVersion : uint8
-	{
-		InitialVersion = 0,
-		// -----<new versions can be added above this line>-----
-		VersionPlusOne,
-		LatestVersion = VersionPlusOne - 1
-	};
+	Ar.UsingCustomVersion(GInstancedStructCustomVersion.GUID);
 
-	EVersion Version = EVersion::LatestVersion;
-
-	// Temporary code to introduce versioning and load old data
-	// The goal is to remove this  by bumping the version in a near future.
-	bool bUseVersioning = true;
-
-#if WITH_EDITOR
-	if (!Ar.IsCooking() && !Ar.IsFilterEditorOnly())
+	if (Ar.IsLoading())
 	{
-		// Keep archive position to use legacy serialization if the header is not found
-		const int64 HeaderOffset = Ar.Tell();
-
-		// Some random pattern to differentiate old data
-		const uint32 NewVersionHeader = 0xABABABAB;
-		uint32 Header = NewVersionHeader;
-		Ar << Header;
-
-		if (Ar.IsLoading())
+		const int32 CustomVersion = Ar.CustomVer(GInstancedStructCustomVersion.GUID);
+		if (CustomVersion < FInstancedStructCustomVersion::CustomVersionAdded)
 		{
-			if (Header != NewVersionHeader)
+			// The old format had "header+version" in editor builds, and just "version" otherwise.
+			// If the first thing we read is the old header, consume it, if not go back and assume that we have just the version.
+			const int64 HeaderOffset = Ar.Tell();
+			uint32 Header = 0;
+			Ar << Header;
+
+			constexpr uint32 LegacyEditorHeader = 0xABABABAB;
+			if (Header != LegacyEditorHeader)
 			{
-				// Not a valid header so go back and process with legacy loading
 				Ar.Seek(HeaderOffset);
-				bUseVersioning = false;
-				UE_LOG(LogLoad, Verbose, TEXT("Loading FInstancedStruct using legacy serialization"));
 			}
-		}
-	}
-#endif // WITH_EDITOR
 
-	if (bUseVersioning)
-	{
-		Ar << Version;
+			uint8 Version = 0;
+			Ar << Version;
+		}
 	}
 
-	if (Version > EVersion::LatestVersion)
-	{
-		UE_LOG(LogCore, Error, TEXT("Invalid Version: %hhu"), int(Version));
-		Ar.SetError();
-		return false;
-	}
+	UScriptStruct* NonConstStruct = const_cast<UScriptStruct*>(GetScriptStruct());
 
 	if (Ar.IsLoading())
 	{
@@ -193,10 +191,7 @@ bool FInstancedStruct::Serialize(FArchive& Ar)
 
 		// Size of the serialized memory
 		int32 SerialSize = 0; 
-		if (bUseVersioning)
-		{
-			Ar << SerialSize;
-		}
+		Ar << SerialSize;
 
 		// Serialized memory
 		if (NonConstStruct == nullptr && SerialSize > 0)
