Fix nopch compilation

diff --git a/Engine/Plugins/Experimental/StructUtils/Source/StructUtils/Private/InstancedStruct.cpp b/Engine/Plugins/Experimental/StructUtils/Source/StructUtils/Private/InstancedStruct.cpp
--- a/Engine/Plugins/Experimental/StructUtils/Source/StructUtils/Private/InstancedStruct.cpp
+++ b/Engine/Plugins/Experimental/StructUtils/Source/StructUtils/Private/InstancedStruct.cpp
@@ -1,6 +1,7 @@
 // Copyright Epic Games, Inc. All Rights Reserved.
 #include "InstancedStruct.h"
 #include "StructView.h"
+#include "Serialization/CustomVersion.h"
 #include "Serialization/PropertyLocalizationDataGathering.h"
 #include "StructUtilsTypes.h"
 
