PropertyBag : Fixed FPropertyBagContainerTypes ranged for loop helpers (after changing an array to a TStaticArray)

diff --git a/Engine/Plugins/Experimental/StructUtils/Source/StructUtils/Public/PropertyBag.h b/Engine/Plugins/Experimental/StructUtils/Source/StructUtils/Public/PropertyBag.h
--- a/Engine/Plugins/Experimental/StructUtils/Source/StructUtils/Public/PropertyBag.h
+++ b/Engine/Plugins/Experimental/StructUtils/Source/StructUtils/Public/PropertyBag.h
@@ -141,10 +141,10 @@ struct STRUCTUTILS_API FPropertyBagContainerTypes
 		return GetArrayHash(PropertyBagContainerTypes.Types.GetData(), PropertyBagContainerTypes.NumContainers);
 	}
 
-	EPropertyBagContainerType* begin() { return &Types[0]; }
-	const EPropertyBagContainerType* begin() const { return &Types[0]; }
-	EPropertyBagContainerType* end()  { return &Types[NumContainers]; }
-	const EPropertyBagContainerType* end() const { return &Types[NumContainers]; }
+	EPropertyBagContainerType* begin() { return Types.GetData(); }
+	const EPropertyBagContainerType* begin() const { return Types.GetData(); }
+	EPropertyBagContainerType* end()  { return Types.GetData() + NumContainers; }
+	const EPropertyBagContainerType* end() const { return Types.GetData() + NumContainers; }
 
 protected:
 	static constexpr uint8 MaxNestedTypes = 2;
