[StructUtils] Made instanced struct filter public so that it can be reused, removed unused code.

diff --git a/Engine/Plugins/Experimental/StructUtils/Source/StructUtilsEditor/Private/Customizations/InstancedStructDetails.cpp b/Engine/Plugins/Experimental/StructUtils/Source/StructUtilsEditor/Private/Customizations/InstancedStructDetails.cpp
--- a/Engine/Plugins/Experimental/StructUtils/Source/StructUtilsEditor/Private/Customizations/InstancedStructDetails.cpp
+++ b/Engine/Plugins/Experimental/StructUtils/Source/StructUtilsEditor/Private/Customizations/InstancedStructDetails.cpp
@@ -11,7 +11,6 @@
 #include "ScopedTransaction.h"
 #include "Modules/ModuleManager.h"
 #include "StructViewerModule.h"
-#include "StructViewerFilter.h"
 #include "Styling/SlateIconFinder.h"
 #include "Engine/UserDefinedStruct.h"
 #include "InstancedStruct.h"
@@ -162,45 +161,32 @@ protected:
 
 ////////////////////////////////////
 
-class FInstancedStructFilter : public IStructViewerFilter
+bool FInstancedStructFilter::IsStructAllowed(const FStructViewerInitializationOptions& InInitOptions, const UScriptStruct* InStruct, TSharedRef<FStructViewerFilterFuncs> InFilterFuncs)
 {
-public:
-	/** The base struct for the property that classes must be a child-of. */
-	const UScriptStruct* BaseStruct = nullptr;
-
-	// A flag controlling whether we allow UserDefinedStructs
-	bool bAllowUserDefinedStructs = false;
-
-	// A flag controlling whether we allow to select the BaseStruct
-	bool bAllowBaseStruct = true;
-
-	virtual bool IsStructAllowed(const FStructViewerInitializationOptions& InInitOptions, const UScriptStruct* InStruct, TSharedRef<FStructViewerFilterFuncs> InFilterFuncs) override
+	if (InStruct->IsA<UUserDefinedStruct>())
 	{
-		if (InStruct->IsA<UUserDefinedStruct>())
-		{
-			return bAllowUserDefinedStructs;
-		}
-
-		if (InStruct == BaseStruct)
-		{
-			return bAllowBaseStruct;
-		}
-
-		if (InStruct->HasMetaData(TEXT("Hidden")))
-		{
-			return false;
-		}
+		return bAllowUserDefinedStructs;
+	}
 
-		// Query the native struct to see if it has the correct parent type (if any)
-		return !BaseStruct || InStruct->IsChildOf(BaseStruct);
+	if (InStruct == BaseStruct)
+	{
+		return bAllowBaseStruct;
 	}
 
-	virtual bool IsUnloadedStructAllowed(const FStructViewerInitializationOptions& InInitOptions, const FSoftObjectPath& InStructPath, TSharedRef<FStructViewerFilterFuncs> InFilterFuncs) override
+	if (InStruct->HasMetaData(TEXT("Hidden")))
 	{
-		// User Defined Structs don't support inheritance, so only include them requested
-		return bAllowUserDefinedStructs;
+		return false;
 	}
-};
+
+	// Query the native struct to see if it has the correct parent type (if any)
+	return !BaseStruct || InStruct->IsChildOf(BaseStruct);
+}
+
+bool FInstancedStructFilter::IsUnloadedStructAllowed(const FStructViewerInitializationOptions& InInitOptions, const FSoftObjectPath& InStructPath, TSharedRef<FStructViewerFilterFuncs> InFilterFuncs)
+{
+	// User Defined Structs don't support inheritance, so only include them requested
+	return bAllowUserDefinedStructs;
+}
 
 ////////////////////////////////////
 
@@ -351,21 +337,9 @@ FName FInstancedStructDataDetails::GetName() const
 	return Name;
 }
 
-void FInstancedStructDataDetails::PostUndo(bool bSuccess)
-{
-	// Undo; force a sync next Tick
-	LastSyncEditableInstanceFromSourceSeconds = 0.0;
-}
-
-void FInstancedStructDataDetails::PostRedo(bool bSuccess)
-{
-	// Redo; force a sync next Tick
-	LastSyncEditableInstanceFromSourceSeconds = 0.0;
-}
 
 ////////////////////////////////////
 
-
 TSharedRef<IPropertyTypeCustomization> FInstancedStructDetails::MakeInstance()
 {
 	return MakeShared<FInstancedStructDetails>();
diff --git a/Engine/Plugins/Experimental/StructUtils/Source/StructUtilsEditor/Public/InstancedStructDetails.h b/Engine/Plugins/Experimental/StructUtils/Source/StructUtilsEditor/Public/InstancedStructDetails.h
--- a/Engine/Plugins/Experimental/StructUtils/Source/StructUtilsEditor/Public/InstancedStructDetails.h
+++ b/Engine/Plugins/Experimental/StructUtils/Source/StructUtilsEditor/Public/InstancedStructDetails.h
@@ -5,6 +5,7 @@
 #include "IPropertyTypeCustomization.h"
 #include "IDetailCustomNodeBuilder.h"
 #include "EditorUndoClient.h"
+#include "StructViewerFilter.h"
 
 class IPropertyHandle;
 class IDetailPropertyRow;
@@ -59,11 +60,11 @@ private:
  * Can be used in a implementation of a IPropertyTypeCustomization CustomizeChildren() to display editable FInstancedStruct contents.
  * OnChildRowAdded() is called right after each property is added, which allows the property row to be customizable.
  */
-class STRUCTUTILSEDITOR_API FInstancedStructDataDetails : public IDetailCustomNodeBuilder, public FSelfRegisteringEditorUndoClient, public TSharedFromThis<FInstancedStructDataDetails>
+class STRUCTUTILSEDITOR_API FInstancedStructDataDetails : public IDetailCustomNodeBuilder, public TSharedFromThis<FInstancedStructDataDetails>
 {
 public:
 	FInstancedStructDataDetails(TSharedPtr<IPropertyHandle> InStructProperty);
-	~FInstancedStructDataDetails();
+	virtual ~FInstancedStructDataDetails() override;
 
 	//~ Begin IDetailCustomNodeBuilder interface
 	virtual void SetOnRebuildChildren(FSimpleDelegate InOnRegenerateChildren) override;
@@ -75,10 +76,6 @@ public:
 	virtual FName GetName() const override;
 	//~ End IDetailCustomNodeBuilder interface
 
-	/** FEditorUndoClient interface */
-	virtual void PostUndo(bool bSuccess) override;
-	virtual void PostRedo(bool bSuccess) override;
-
 	// Called when a child is added, override to customize a child row.
 	virtual void OnChildRowAdded(IDetailPropertyRow& ChildRow) {}
 
@@ -105,15 +102,32 @@ private:
 	/** Delegate that can be used to refresh the child rows of the current struct (eg, when changing struct type) */
 	FSimpleDelegate OnRegenerateChildren;
 
-	/** The last time that SyncEditableInstanceFromSource was called, in FPlatformTime::Seconds() */
-	double LastSyncEditableInstanceFromSourceSeconds = 0.0;
-
 	/** True if we're currently handling a StructValuePostChange */
 	bool bIsHandlingStructValuePostChange = false;
 	
 	FDelegateHandle UserDefinedStructReinstancedHandle;
 };
 
+/**
+ * Filter used by the instanced struct struct picker.
+ */
+ class STRUCTUTILSEDITOR_API FInstancedStructFilter : public IStructViewerFilter
+{
+public:
+	/** The base struct for the property that classes must be a child-of. */
+	const UScriptStruct* BaseStruct = nullptr;
+
+	// A flag controlling whether we allow UserDefinedStructs
+	bool bAllowUserDefinedStructs = false;
+
+	// A flag controlling whether we allow to select the BaseStruct
+	bool bAllowBaseStruct = true;
+
+	virtual bool IsStructAllowed(const FStructViewerInitializationOptions& InInitOptions, const UScriptStruct* InStruct, TSharedRef<FStructViewerFilterFuncs> InFilterFuncs) override;
+	virtual bool IsUnloadedStructAllowed(const FStructViewerInitializationOptions& InInitOptions, const FSoftObjectPath& InStructPath, TSharedRef<FStructViewerFilterFuncs> InFilterFuncs) override;
+};
+
+
 #if UE_ENABLE_INCLUDE_ORDER_DEPRECATED_IN_5_2
 #include "CoreMinimal.h"
 #include "UnrealClient.h"
