Add ability to create specific instances of a ScriptStruct from a PropertyBag
* This change adds a funciton to PropertyBag to allow a caller to manually specify the name of the struct that gets generated.
* TEDS queries depend on matching specific types to have an effect.  The Memento system can generate memento structs that need to have a one-to-one mapping to the mementoizable column to perform the translation between the two types.  This requires that the generated type for column structs with two identical layouts be unique in order to create the correct processor.

diff --git a/Engine/Plugins/Experimental/StructUtils/Source/StructUtils/Private/PropertyBag.cpp b/Engine/Plugins/Experimental/StructUtils/Source/StructUtils/Private/PropertyBag.cpp
--- a/Engine/Plugins/Experimental/StructUtils/Source/StructUtils/Private/PropertyBag.cpp
+++ b/Engine/Plugins/Experimental/StructUtils/Source/StructUtils/Private/PropertyBag.cpp
@@ -2274,11 +2274,13 @@ EPropertyBagResult FPropertyBagArrayRef::SetValueSoftPath(const int32 Index, con
 //  UPropertyBag
 //----------------------------------------------------------------//
 
-const UPropertyBag* UPropertyBag::GetOrCreateFromDescs(const TConstArrayView<FPropertyBagPropertyDesc> PropertyDescs)
+const UPropertyBag* UPropertyBag::GetOrCreateFromDescs(const TConstArrayView<FPropertyBagPropertyDesc> PropertyDescs, const TCHAR* PrefixName)
 {
 	const uint64 BagHash = UE::StructUtils::Private::CalcPropertyDescArrayHash(PropertyDescs);
-	const FString ScriptStructName = FString::Printf(TEXT("PropertyBag_%llx"), BagHash);
-
+	const FString ScriptStructName = PrefixName == nullptr
+		? FString::Printf(TEXT("PropertyBag_%llx"), BagHash)
+		: FString::Printf(TEXT("%s_%llx"), PrefixName, BagHash);
+	
 	if (const UPropertyBag* ExistingBag = FindObject<UPropertyBag>(GetTransientPackage(), *ScriptStructName))
 	{
 		return ExistingBag;
diff --git a/Engine/Plugins/Experimental/StructUtils/Source/StructUtils/Public/PropertyBag.h b/Engine/Plugins/Experimental/StructUtils/Source/StructUtils/Public/PropertyBag.h
--- a/Engine/Plugins/Experimental/StructUtils/Source/StructUtils/Public/PropertyBag.h
+++ b/Engine/Plugins/Experimental/StructUtils/Source/StructUtils/Public/PropertyBag.h
@@ -835,10 +835,14 @@ public:
 	GENERATED_BODY()
 
 	/**
-	 * Creates new UPropertyBag struct based on the properties passed in.
-	 * If there are multiple properties that have the same name, only the first one is added.
+	 * Returns UPropertyBag struct based on the property descriptions passed in.
+	 * UPropertyBag struct names will be auto-generated by prefixing 'PropertyBag_' to the hash of the descriptions.
+	 * If a UPropertyBag with same name already exists, the existing object is returned.
+	 * This means that a property bags which share same layout (same descriptions) will share the same UPropertyBag.
+	 * If there are multiple properties that have the same name, only the first property is added.
+	 * The caller is expected to ensure unique names for the property descriptions.
 	 */
-	static const UPropertyBag* GetOrCreateFromDescs(const TConstArrayView<FPropertyBagPropertyDesc> InPropertyDescs);
+	static const UPropertyBag* GetOrCreateFromDescs(const TConstArrayView<FPropertyBagPropertyDesc> InPropertyDescs, const TCHAR* PrefixName = nullptr);
 
 	/** Returns property descriptions that specify this struct. */
 	TConstArrayView<FPropertyBagPropertyDesc> GetPropertyDescs() const { return PropertyDescs; }
diff --git a/Engine/Plugins/Experimental/TypedElementsDataStorage/Source/TypedElementsDataStorage/Private/Memento/TypedElementMementoTranslators.cpp b/Engine/Plugins/Experimental/TypedElementsDataStorage/Source/TypedElementsDataStorage/Private/Memento/TypedElementMementoTranslators.cpp
--- a/Engine/Plugins/Experimental/TypedElementsDataStorage/Source/TypedElementsDataStorage/Private/Memento/TypedElementMementoTranslators.cpp
+++ b/Engine/Plugins/Experimental/TypedElementsDataStorage/Source/TypedElementsDataStorage/Private/Memento/TypedElementMementoTranslators.cpp
@@ -30,10 +30,9 @@ void UTypedElementDefaultMementoTranslator::PostInitProperties()
 	{
 		PropertyDescs.Add(FPropertyBagPropertyDesc(Property->GetFName(), Property));
 	}
-
-	// Note: The name of the memento will be an unhelpful "PropertyBag_<hash>"
-	// May wish to revisit property bag interface to allow name to be specified for debug purposes
-	const UPropertyBag* PropertyBag = UPropertyBag::GetOrCreateFromDescs(PropertyDescs);
+	
+	const FString GeneratedMementoStructName = FString::Printf(TEXT("%s_Memento"), *SourceColumnType->GetName());
+	const UPropertyBag* PropertyBag = UPropertyBag::GetOrCreateFromDescs(PropertyDescs, *GeneratedMementoStructName);
 	MementoType = PropertyBag;
 
 	// Need to change the type to a FTypedElementDataStorageColumn to appease TEDS/Mass
