[StructUtils][PCG] Improve filtering for pin types supported in Property Bags

* IsPinTypeAccepted UFunction is now receiving a bool that will tell if the pin type is a child or not
* It will allow to filter out specific structures
* On PCG side, add the support for Vec2/Vec4/Quat/SoftObjectPath/SoftClassPath as graph parameters

#rnx

diff --git a/Engine/Plugins/Experimental/PCG/Source/PCG/Private/PCGGraph.cpp b/Engine/Plugins/Experimental/PCG/Source/PCG/Private/PCGGraph.cpp
--- a/Engine/Plugins/Experimental/PCG/Source/PCG/Private/PCGGraph.cpp
+++ b/Engine/Plugins/Experimental/PCG/Source/PCG/Private/PCGGraph.cpp
@@ -1089,7 +1089,7 @@ void UPCGGraph::OnGraphParametersChanged(EPCGGraphParameterEvent InChangeType, F
 	NotifyGraphParametersChanged(InChangeType, InChangedPropertyName);
 }
 
-bool UPCGGraph::UserParametersIsPinTypeAccepted(FEdGraphPinType InPinType)
+bool UPCGGraph::UserParametersIsPinTypeAccepted(FEdGraphPinType InPinType, bool bIsChild)
 {
 	// Text and interface not supported
 	if (InPinType.PinCategory == TEXT("text") || InPinType.PinCategory == TEXT("interface"))
@@ -1098,10 +1098,24 @@ bool UPCGGraph::UserParametersIsPinTypeAccepted(FEdGraphPinType InPinType)
 	}
 	else if (InPinType.PinCategory == TEXT("struct"))
 	{
-		// Structs other than Vector/Transform/Rotator not supported
-		return InPinType.PinSubCategoryObject == TBaseStructure<FVector>::Get() 
-			|| InPinType.PinSubCategoryObject == TBaseStructure<FTransform>::Get()
-			|| InPinType.PinSubCategoryObject == TBaseStructure<FRotator>::Get();
+		// Vector/Transform/Rotator are "explicitly" defined. Also return true for a "null" PinSubCategoryObject
+		// to mark it as "valid" so the system will call this function back on all available structures.
+		if (!bIsChild)
+		{
+			return InPinType.PinSubCategoryObject.IsExplicitlyNull()
+				|| InPinType.PinSubCategoryObject == TBaseStructure<FVector>::Get()
+				|| InPinType.PinSubCategoryObject == TBaseStructure<FTransform>::Get()
+				|| InPinType.PinSubCategoryObject == TBaseStructure<FRotator>::Get();
+		}
+		else
+		{
+			// Since we are not supporting all the structures by default, only show the ones that we are supporting with PCG.
+			return InPinType.PinSubCategoryObject == TBaseStructure<FVector2D>::Get()
+				|| InPinType.PinSubCategoryObject == TBaseStructure<FVector4>::Get()
+				|| InPinType.PinSubCategoryObject == TBaseStructure<FQuat>::Get()
+				|| InPinType.PinSubCategoryObject == TBaseStructure<FSoftObjectPath>::Get()
+				|| InPinType.PinSubCategoryObject == TBaseStructure<FSoftClassPath>::Get();
+		}
 	}
 	else
 	{
diff --git a/Engine/Plugins/Experimental/PCG/Source/PCG/Public/PCGGraph.h b/Engine/Plugins/Experimental/PCG/Source/PCG/Public/PCGGraph.h
--- a/Engine/Plugins/Experimental/PCG/Source/PCG/Public/PCGGraph.h
+++ b/Engine/Plugins/Experimental/PCG/Source/PCG/Public/PCGGraph.h
@@ -289,7 +289,7 @@ protected:
 
 #if WITH_EDITOR
 	UFUNCTION(BlueprintInternalUseOnly)
-	bool UserParametersIsPinTypeAccepted(FEdGraphPinType InPinType);
+	bool UserParametersIsPinTypeAccepted(FEdGraphPinType InPinType, bool bIsChild);
 	
 	UFUNCTION(BlueprintInternalUseOnly)
 	bool UserParametersCanRemoveProperty(FGuid InPropertyID, FName InPropertyName);
diff --git a/Engine/Plugins/Experimental/StructUtils/Source/StructUtilsEditor/Private/Customizations/PropertyBagDetails.cpp b/Engine/Plugins/Experimental/StructUtils/Source/StructUtilsEditor/Private/Customizations/PropertyBagDetails.cpp
--- a/Engine/Plugins/Experimental/StructUtils/Source/StructUtilsEditor/Private/Customizations/PropertyBagDetails.cpp
+++ b/Engine/Plugins/Experimental/StructUtils/Source/StructUtilsEditor/Private/Customizations/PropertyBagDetails.cpp
@@ -874,19 +874,19 @@ TSharedRef<SWidget> FPropertyBagInstanceDataDetails::OnPropertyNameContent(TShar
 		{
 			check(IsPinTypeAcceptedFunc && IsPinTypeAcceptedTarget);
 
-			// We need to make sure the signature matches perfectly: bool(FEdGraphPinType)
-			bool bFuncIsValid = UE::StructUtils::Private::ValidateFunctionSignature<bool, FEdGraphPinType>(IsPinTypeAcceptedFunc);
+			// We need to make sure the signature matches perfectly: bool(FEdGraphPinType, bool)
+			bool bFuncIsValid = UE::StructUtils::Private::ValidateFunctionSignature<bool, FEdGraphPinType, bool>(IsPinTypeAcceptedFunc);
 			if (!ensureMsgf(bFuncIsValid, TEXT("[%s] Function %s does not have the right signature."), *UE::StructUtils::Metadata::IsPinTypeAcceptedName.ToString(), *IsPinTypeAcceptedFunc->GetName()))
 			{
 				return;
 			}
 		}
 
-		auto IsPinTypeAccepted = [IsPinTypeAcceptedFunc, IsPinTypeAcceptedTarget](const FEdGraphPinType& InPinType) -> bool
+		auto IsPinTypeAccepted = [IsPinTypeAcceptedFunc, IsPinTypeAcceptedTarget](const FEdGraphPinType& InPinType, bool bInIsChild) -> bool
 		{
 			if (IsPinTypeAcceptedFunc && IsPinTypeAcceptedTarget)
 			{
-				const TValueOrError<bool, void> bIsValid = UE::StructUtils::Private::CallFunc<bool>(IsPinTypeAcceptedTarget, IsPinTypeAcceptedFunc, InPinType);
+				const TValueOrError<bool, void> bIsValid = UE::StructUtils::Private::CallFunc<bool>(IsPinTypeAcceptedTarget, IsPinTypeAcceptedFunc, InPinType, bInIsChild);
 				return bIsValid.HasValue() && bIsValid.GetValue();
 			}
 			else
@@ -902,7 +902,7 @@ TSharedRef<SWidget> FPropertyBagInstanceDataDetails::OnPropertyNameContent(TShar
 		// Filter
 		for (TSharedPtr<UEdGraphSchema_K2::FPinTypeTreeInfo>& PinType : TempTypeTree)
 		{
-			if (!PinType.IsValid() || !IsPinTypeAccepted(PinType->GetPinType(/*bForceLoadSubCategoryObject*/false)))
+			if (!PinType.IsValid() || !IsPinTypeAccepted(PinType->GetPinType(/*bForceLoadSubCategoryObject*/false), /*bInIsChild=*/ false))
 			{
 				continue;
 			}
@@ -914,7 +914,7 @@ TSharedRef<SWidget> FPropertyBagInstanceDataDetails::OnPropertyNameContent(TShar
 				{
 					const FEdGraphPinType& ChildPinType = Child->GetPinType(/*bForceLoadSubCategoryObject*/false);
 
-					if (!UE::StructUtils::Private::CanHaveMemberVariableOfType(ChildPinType) || !IsPinTypeAccepted(ChildPinType))
+					if (!UE::StructUtils::Private::CanHaveMemberVariableOfType(ChildPinType) || !IsPinTypeAccepted(ChildPinType, /*bInIsChild=*/ true))
 					{
 						PinType->Children.RemoveAt(ChildIndex);
 						continue;
