Do not use Cast as a const_cast
#rnx

diff --git a/Engine/Plugins/Experimental/StructUtils/Source/StructUtils/Private/PropertyBag.cpp b/Engine/Plugins/Experimental/StructUtils/Source/StructUtils/Private/PropertyBag.cpp
--- a/Engine/Plugins/Experimental/StructUtils/Source/StructUtils/Private/PropertyBag.cpp
+++ b/Engine/Plugins/Experimental/StructUtils/Source/StructUtils/Private/PropertyBag.cpp
@@ -171,7 +171,7 @@ namespace UE::StructUtils::Private
 		return EPropertyBagPropertyType::None;
 	}
 
-	UObject* GetValueTypeObjectFromProperty(const FProperty* InSourceProperty)
+	const UObject* GetValueTypeObjectFromProperty(const FProperty* InSourceProperty)
 	{
 		if (const FByteProperty* ByteProperty = CastField<FByteProperty>(InSourceProperty))
 		{
@@ -314,20 +314,20 @@ namespace UE::StructUtils::Private
 				return Prop;
 			}
 		case EPropertyBagPropertyType::Enum:
-			if (UEnum* Enum = Cast<UEnum>(Desc.ValueTypeObject))
+			if (const UEnum* Enum = Cast<UEnum>(Desc.ValueTypeObject))
 			{
 				FEnumProperty* Prop = new FEnumProperty(PropertyScope, Desc.Name, RF_Public);
 				FNumericProperty* UnderlyingProp = new FByteProperty(Prop, "UnderlyingType", RF_Public); // HACK: Hardwire to byte property for now for BP compatibility
-				Prop->SetEnum(Enum);
+				Prop->SetEnum(const_cast<UEnum*>(Enum));
 				Prop->AddCppProperty(UnderlyingProp);
 				return Prop;
 			}
 			break;
 		case EPropertyBagPropertyType::Struct:
-			if (UScriptStruct* ScriptStruct = Cast<UScriptStruct>(Desc.ValueTypeObject))
+			if (const UScriptStruct* ScriptStruct = Cast<UScriptStruct>(Desc.ValueTypeObject))
 			{
 				FStructProperty* Prop = new FStructProperty(PropertyScope, Desc.Name, RF_Public);
-				Prop->Struct = ScriptStruct;
+				Prop->Struct = const_cast<UScriptStruct*>(ScriptStruct);
 
 				if (ScriptStruct->GetCppStructOps() && ScriptStruct->GetCppStructOps()->HasGetTypeHash())
 				{
@@ -343,46 +343,46 @@ namespace UE::StructUtils::Private
 			}
 			break;
 		case EPropertyBagPropertyType::Object:
-			if (UClass* Class = Cast<UClass>(Desc.ValueTypeObject))
+			if (const UClass* Class = Cast<UClass>(Desc.ValueTypeObject))
 			{
 				FObjectProperty* Prop = new FObjectProperty(PropertyScope, Desc.Name, RF_Public);
 				if (Class->HasAnyClassFlags(CLASS_DefaultToInstanced))
 				{
 					Prop->SetPropertyFlags(CPF_InstancedReference);
 				}
-				Prop->SetPropertyClass(Class);
+				Prop->SetPropertyClass(const_cast<UClass*>(Class));
 				Prop->SetPropertyFlags(CPF_HasGetValueTypeHash);
 				return Prop;
 			}
 			break;
 		case EPropertyBagPropertyType::SoftObject:
-			if (UClass* Class = Cast<UClass>(Desc.ValueTypeObject))
+			if (const UClass* Class = Cast<UClass>(Desc.ValueTypeObject))
 			{
 				FSoftObjectProperty* Prop = new FSoftObjectProperty(PropertyScope, Desc.Name, RF_Public);
 				if (Class->HasAnyClassFlags(CLASS_DefaultToInstanced))
 				{
 					Prop->SetPropertyFlags(CPF_InstancedReference);
 				}
-				Prop->SetPropertyClass(Class);
+				Prop->SetPropertyClass(const_cast<UClass*>(Class));
 				Prop->SetPropertyFlags(CPF_HasGetValueTypeHash);
 				return Prop;
 			}
 			break;
 		case EPropertyBagPropertyType::Class:
-			if (UClass* Class = Cast<UClass>(Desc.ValueTypeObject))
+			if (const UClass* Class = Cast<UClass>(Desc.ValueTypeObject))
 			{
 				FClassProperty* Prop = new FClassProperty(PropertyScope, Desc.Name, RF_Public);
-				Prop->SetMetaClass(Class);
+				Prop->SetMetaClass(const_cast<UClass*>(Class));
 				Prop->PropertyClass = UClass::StaticClass();
 				Prop->SetPropertyFlags(CPF_HasGetValueTypeHash);
 				return Prop;
 			}
 			break;
 		case EPropertyBagPropertyType::SoftClass:
-			if (UClass* Class = Cast<UClass>(Desc.ValueTypeObject))
+			if (const UClass* Class = Cast<UClass>(Desc.ValueTypeObject))
 			{
 				FSoftClassProperty* Prop = new FSoftClassProperty(PropertyScope, Desc.Name, RF_Public);
-				Prop->SetMetaClass(Class);
+				Prop->SetMetaClass(const_cast<UClass*>(Class));
 				Prop->PropertyClass = UClass::StaticClass();
 				Prop->SetPropertyFlags(CPF_HasGetValueTypeHash);
 				return Prop;
@@ -1946,7 +1946,7 @@ void FInstancedPropertyBag::AddStructReferencedObjects(FReferenceCollector& Coll
 				{
 					if (Desc.ValueType == EPropertyBagPropertyType::Struct)
 					{
-						if (UUserDefinedStruct* UserDefinedStruct = Cast<UUserDefinedStruct>(Desc.ValueTypeObject))
+						if (const UUserDefinedStruct* UserDefinedStruct = Cast<UUserDefinedStruct>(Desc.ValueTypeObject))
 						{
 							if (UserDefinedStruct == StructureToReinstance->PrimaryStruct)
 							{
