StateTree: Fix undoing property add immediately after adding it
- Added Identical() to FInstancedPropertyBag
- Fixed FInstancedPropertyBag to reset when loading empty data

diff --git a/Engine/Plugins/Experimental/StructUtils/Source/StructUtils/Private/PropertyBag.cpp b/Engine/Plugins/Experimental/StructUtils/Source/StructUtils/Private/PropertyBag.cpp
--- a/Engine/Plugins/Experimental/StructUtils/Source/StructUtils/Private/PropertyBag.cpp
+++ b/Engine/Plugins/Experimental/StructUtils/Source/StructUtils/Private/PropertyBag.cpp
@@ -1825,6 +1825,11 @@ TValueOrError<const FPropertyBagArrayRef, EPropertyBagResult> FInstancedProperty
 	return MakeValue(FPropertyBagArrayRef(*Desc, Address));
 }
 
+bool FInstancedPropertyBag::Identical(const FInstancedPropertyBag* Other, const uint32 PortFlags) const
+{
+	return Other && Value.Identical(&Other->Value, PortFlags);
+}
+
 bool FInstancedPropertyBag::Serialize(FArchive& Ar)
 {
 	// Obsolete, use custom version instead.
@@ -1923,6 +1928,14 @@ bool FInstancedPropertyBag::Serialize(FArchive& Ar)
 			Ar.Seek(FinalOffset);	// Reset archive to its position
 		}
 	}
+	else
+	{
+		if (Ar.IsLoading())
+		{
+			// If loading and there was no data saved in the archive, make sure the value is empty.
+			Reset();
+		}
+	}
 	
 	return true;
 }
diff --git a/Engine/Plugins/Experimental/StructUtils/Source/StructUtils/Public/PropertyBag.h b/Engine/Plugins/Experimental/StructUtils/Source/StructUtils/Public/PropertyBag.h
--- a/Engine/Plugins/Experimental/StructUtils/Source/StructUtils/Public/PropertyBag.h
+++ b/Engine/Plugins/Experimental/StructUtils/Source/StructUtils/Public/PropertyBag.h
@@ -589,6 +589,7 @@ struct STRUCTUTILS_API FInstancedPropertyBag
 	*/
 	TValueOrError<const FPropertyBagArrayRef, EPropertyBagResult> GetArrayRef(const FName Name) const;
 
+	bool Identical(const FInstancedPropertyBag* Other, uint32 PortFlags) const;
 	bool Serialize(FArchive& Ar);
 	void AddStructReferencedObjects(FReferenceCollector& Collector);
 
@@ -604,6 +605,7 @@ template<> struct TStructOpsTypeTraits<FInstancedPropertyBag> : public TStructOp
 {
 	enum
 	{
+		WithIdentical = true,
 		WithSerializer = true,
 		WithAddStructReferencedObjects = true,
 	};
