PropertyBag : Fis for CIS error UUSerStruct missing dependencies on load

diff --git a/Engine/Plugins/Experimental/StructUtils/Source/StructUtils/Private/PropertyBag.cpp b/Engine/Plugins/Experimental/StructUtils/Source/StructUtils/Private/PropertyBag.cpp
--- a/Engine/Plugins/Experimental/StructUtils/Source/StructUtils/Private/PropertyBag.cpp
+++ b/Engine/Plugins/Experimental/StructUtils/Source/StructUtils/Private/PropertyBag.cpp
@@ -1863,6 +1863,14 @@ bool FInstancedPropertyBag::Serialize(FArchive& Ar)
 			TArray<FPropertyBagPropertyDesc> PropertyDescs;
 			Ar << PropertyDescs;
 
+			for (FPropertyBagPropertyDesc& PropDesc : PropertyDescs)
+			{
+				if (PropDesc.ValueTypeObject)
+				{
+					Ar.Preload(const_cast<UObject*>(PropDesc.ValueTypeObject.Get()));
+				}
+			}
+			
 			BagStruct = const_cast<UPropertyBag*>(UPropertyBag::GetOrCreateFromDescs(PropertyDescs));
 			Value.InitializeAs(BagStruct);
 
@@ -2005,6 +2013,31 @@ void FInstancedPropertyBag::AddStructReferencedObjects(FReferenceCollector& Coll
 #endif	
 }
 
+void FInstancedPropertyBag::GetPreloadDependencies(TArray<UObject*>& OutDeps)
+{
+	if (const UPropertyBag* BagStruct = GetPropertyBagStruct())
+	{
+		for (const FPropertyBagPropertyDesc& Desc : BagStruct->PropertyDescs)
+		{
+			if (Desc.ValueTypeObject)
+			{
+				OutDeps.Add(const_cast<UObject*>(Desc.ValueTypeObject.Get()));
+			}
+		}
+
+		// Report indirect dependencies of the instanced property bag struct
+		// The iterator will recursively loop through all structs in structs/containers too
+		for (TPropertyValueIterator<FStructProperty> It(BagStruct, Value.GetMutableMemory()); It; ++It)
+		{
+			const UScriptStruct* StructType = It.Key()->Struct;
+			if (UScriptStruct::ICppStructOps* CppStructOps = StructType->GetCppStructOps())
+			{
+				void* StructDataPtr = const_cast<void*>(It.Value());
+				CppStructOps->GetPreloadDependencies(StructDataPtr, OutDeps);
+			}
+		}
+	}
+}
 
 //----------------------------------------------------------------//
 //  FPropertyBagArrayRef
diff --git a/Engine/Plugins/Experimental/StructUtils/Source/StructUtils/Public/PropertyBag.h b/Engine/Plugins/Experimental/StructUtils/Source/StructUtils/Public/PropertyBag.h
--- a/Engine/Plugins/Experimental/StructUtils/Source/StructUtils/Public/PropertyBag.h
+++ b/Engine/Plugins/Experimental/StructUtils/Source/StructUtils/Public/PropertyBag.h
@@ -592,6 +592,7 @@ struct STRUCTUTILS_API FInstancedPropertyBag
 	bool Identical(const FInstancedPropertyBag* Other, uint32 PortFlags) const;
 	bool Serialize(FArchive& Ar);
 	void AddStructReferencedObjects(FReferenceCollector& Collector);
+	void GetPreloadDependencies(TArray<UObject*>& OutDeps);
 
 protected:
 	const void* GetValueAddress(const FPropertyBagPropertyDesc* Desc) const;
@@ -608,6 +609,7 @@ template<> struct TStructOpsTypeTraits<FInstancedPropertyBag> : public TStructOp
 		WithIdentical = true,
 		WithSerializer = true,
 		WithAddStructReferencedObjects = true,
+		WithGetPreloadDependencies = true,
 	};
 };
 
