Added unsigned int support to  PropertyBag, leaving the new types as unsupported at the UI Details panel (as these types are not supported by Blueprints)
Updated RigVMMemoryStorageStruct to support the new types

diff --git a/Engine/Plugins/Experimental/StructUtils/Source/StructUtils/Private/PropertyBag.cpp b/Engine/Plugins/Experimental/StructUtils/Source/StructUtils/Private/PropertyBag.cpp
--- a/Engine/Plugins/Experimental/StructUtils/Source/StructUtils/Private/PropertyBag.cpp
+++ b/Engine/Plugins/Experimental/StructUtils/Source/StructUtils/Private/PropertyBag.cpp
@@ -113,10 +113,18 @@ namespace UE::StructUtils::Private
 		{
 			return EPropertyBagPropertyType::Int32;
 		}
+		if (CastField<FUInt32Property>(InSourceProperty))
+		{
+			return EPropertyBagPropertyType::UInt32;
+		}
 		if (CastField<FInt64Property>(InSourceProperty))
 		{
 			return EPropertyBagPropertyType::Int64;
 		}
+		if (CastField<FUInt64Property>(InSourceProperty))
+		{
+			return EPropertyBagPropertyType::UInt64;
+		}
 		if (CastField<FFloatProperty>(InSourceProperty))
 		{
 			return EPropertyBagPropertyType::Float;
@@ -278,12 +286,24 @@ namespace UE::StructUtils::Private
 				Prop->SetPropertyFlags(CPF_HasGetValueTypeHash);
 				return Prop;
 			}
+		case EPropertyBagPropertyType::UInt32:
+			{
+				FUInt32Property* Prop = new FUInt32Property(PropertyScope, Desc.Name, RF_Public);
+				Prop->SetPropertyFlags(CPF_HasGetValueTypeHash);
+				return Prop;
+			}
 		case EPropertyBagPropertyType::Int64:
 			{
 				FInt64Property* Prop = new FInt64Property(PropertyScope, Desc.Name, RF_Public);
 				Prop->SetPropertyFlags(CPF_HasGetValueTypeHash);
 				return Prop;
 			}
+		case EPropertyBagPropertyType::UInt64:
+			{
+				FUInt64Property* Prop = new FUInt64Property(PropertyScope, Desc.Name, RF_Public);
+				Prop->SetPropertyFlags(CPF_HasGetValueTypeHash);
+				return Prop;
+			}
 		case EPropertyBagPropertyType::Float:
 			{
 				FFloatProperty* Prop = new FFloatProperty(PropertyScope, Desc.Name, RF_Public);
@@ -436,12 +456,24 @@ namespace UE::StructUtils::Private
 				OutValue = Property->GetPropertyValue(Address);
 				return EPropertyBagResult::Success;
 			}
+		case EPropertyBagPropertyType::UInt32:
+			{
+				const FUInt32Property* Property = CastFieldChecked<FUInt32Property>(Desc->CachedProperty);
+				OutValue = (uint32)Property->GetPropertyValue(Address);
+				return EPropertyBagResult::Success;
+			}
 		case EPropertyBagPropertyType::Int64:
 			{
 				const FInt64Property* Property = CastFieldChecked<FInt64Property>(Desc->CachedProperty);
 				OutValue = Property->GetPropertyValue(Address);
 				return EPropertyBagResult::Success;
 			}
+		case EPropertyBagPropertyType::UInt64:
+			{
+				const FUInt64Property* Property = CastFieldChecked<FUInt64Property>(Desc->CachedProperty);
+				OutValue = (uint64)Property->GetPropertyValue(Address);
+				return EPropertyBagResult::Success;
+			}
 		case EPropertyBagPropertyType::Float:
 			{
 				const FFloatProperty* Property = CastFieldChecked<FFloatProperty>(Desc->CachedProperty);
@@ -467,6 +499,84 @@ namespace UE::StructUtils::Private
 		}
 	}
 
+	EPropertyBagResult GetPropertyAsUInt64(const FPropertyBagPropertyDesc* Desc, const void* Address, uint64& OutValue)
+	{
+		if (Desc == nullptr || Desc->CachedProperty == nullptr)
+		{
+			return EPropertyBagResult::PropertyNotFound;
+		}
+		if (Address == nullptr)
+		{
+			return EPropertyBagResult::OutOfBounds;
+		}
+		if (Desc->ContainerTypes.Num() > 0)
+		{
+			return EPropertyBagResult::TypeMismatch;
+		}
+
+		switch (Desc->ValueType)
+		{
+		case EPropertyBagPropertyType::Bool:
+		{
+			const FBoolProperty* Property = CastFieldChecked<FBoolProperty>(Desc->CachedProperty);
+			OutValue = Property->GetPropertyValue(Address) ? 1 : 0;
+			return EPropertyBagResult::Success;
+		}
+		case EPropertyBagPropertyType::Byte:
+		{
+			const FByteProperty* Property = CastFieldChecked<FByteProperty>(Desc->CachedProperty);
+			OutValue = Property->GetPropertyValue(Address);
+			return EPropertyBagResult::Success;
+		}
+		case EPropertyBagPropertyType::Int32:
+		{
+			const FIntProperty* Property = CastFieldChecked<FIntProperty>(Desc->CachedProperty);
+			OutValue = (uint32)Property->GetPropertyValue(Address);
+			return EPropertyBagResult::Success;
+		}
+		case EPropertyBagPropertyType::UInt32:
+		{
+			const FUInt32Property* Property = CastFieldChecked<FUInt32Property>(Desc->CachedProperty);
+			OutValue = Property->GetPropertyValue(Address);
+			return EPropertyBagResult::Success;
+		}
+		case EPropertyBagPropertyType::Int64:
+		{
+			const FInt64Property* Property = CastFieldChecked<FInt64Property>(Desc->CachedProperty);
+			OutValue = (uint64)Property->GetPropertyValue(Address);
+			return EPropertyBagResult::Success;
+		}
+		case EPropertyBagPropertyType::UInt64:
+		{
+			const FUInt64Property* Property = CastFieldChecked<FUInt64Property>(Desc->CachedProperty);
+			OutValue = Property->GetPropertyValue(Address);
+			return EPropertyBagResult::Success;
+		}
+		case EPropertyBagPropertyType::Float:
+		{
+			const FFloatProperty* Property = CastFieldChecked<FFloatProperty>(Desc->CachedProperty);
+			OutValue = (uint64)Property->GetPropertyValue(Address);
+			return EPropertyBagResult::Success;
+		}
+		case EPropertyBagPropertyType::Double:
+		{
+			const FDoubleProperty* Property = CastFieldChecked<FDoubleProperty>(Desc->CachedProperty);
+			OutValue = (uint64)Property->GetPropertyValue(Address);
+			return EPropertyBagResult::Success;
+		}
+		case EPropertyBagPropertyType::Enum:
+		{
+			const FEnumProperty* EnumProperty = CastFieldChecked<FEnumProperty>(Desc->CachedProperty);
+			const FNumericProperty* UnderlyingProperty = EnumProperty->GetUnderlyingProperty();
+			check(UnderlyingProperty);
+			OutValue = UnderlyingProperty->GetUnsignedIntPropertyValue(Address);
+			return EPropertyBagResult::Success;
+		}
+		default:
+			return EPropertyBagResult::TypeMismatch;
+		}
+	}
+
 	EPropertyBagResult GetPropertyAsDouble(const FPropertyBagPropertyDesc* Desc, const void* Address, double& OutValue)
 	{
 		if (Desc == nullptr || Desc->CachedProperty == nullptr)
@@ -502,12 +612,24 @@ namespace UE::StructUtils::Private
 				OutValue = Property->GetPropertyValue(Address);
 				return EPropertyBagResult::Success;
 			}
+		case EPropertyBagPropertyType::UInt32:
+			{
+				const FUInt32Property* Property = CastFieldChecked<FUInt32Property>(Desc->CachedProperty);
+				OutValue = Property->GetPropertyValue(Address);
+				return EPropertyBagResult::Success;
+			}
 		case EPropertyBagPropertyType::Int64:
 			{
 				const FInt64Property* Property = CastFieldChecked<FInt64Property>(Desc->CachedProperty);
 				OutValue = Property->GetPropertyValue(Address);
 				return EPropertyBagResult::Success;
 			}
+		case EPropertyBagPropertyType::UInt64:
+			{
+				const FUInt64Property* Property = CastFieldChecked<FUInt64Property>(Desc->CachedProperty);
+				OutValue = Property->GetPropertyValue(Address);
+				return EPropertyBagResult::Success;
+			}
 		case EPropertyBagPropertyType::Float:
 			{
 				const FFloatProperty* Property = CastFieldChecked<FFloatProperty>(Desc->CachedProperty);
@@ -728,12 +850,24 @@ namespace UE::StructUtils::Private
 				Property->SetPropertyValue(Address, (int32)InValue);
 				return EPropertyBagResult::Success;
 			}
+		case EPropertyBagPropertyType::UInt32:
+			{
+				const FUInt32Property* Property = CastFieldChecked<FUInt32Property>(Desc->CachedProperty);
+				Property->SetPropertyValue(Address, (uint32)InValue);
+				return EPropertyBagResult::Success;
+			}
 		case EPropertyBagPropertyType::Int64:
 			{
 				const FInt64Property* Property = CastFieldChecked<FInt64Property>(Desc->CachedProperty);
 				Property->SetPropertyValue(Address, InValue);
 				return EPropertyBagResult::Success;
 			}
+		case EPropertyBagPropertyType::UInt64:
+			{
+				const FUInt64Property* Property = CastFieldChecked<FUInt64Property>(Desc->CachedProperty);
+				Property->SetPropertyValue(Address, (uint64)InValue);
+				return EPropertyBagResult::Success;
+			}
 		case EPropertyBagPropertyType::Float:
 			{
 				const FFloatProperty* Property = CastFieldChecked<FFloatProperty>(Desc->CachedProperty);
@@ -759,6 +893,84 @@ namespace UE::StructUtils::Private
 		}
 	}
 
+	EPropertyBagResult SetPropertyFromUInt64(const FPropertyBagPropertyDesc* Desc, void* Address, const uint64 InValue)
+	{
+		if (Desc == nullptr || Desc->CachedProperty == nullptr)
+		{
+			return EPropertyBagResult::PropertyNotFound;
+		}
+		if (Address == nullptr)
+		{
+			return EPropertyBagResult::OutOfBounds;
+		}
+		if (Desc->ContainerTypes.Num() > 0)
+		{
+			return EPropertyBagResult::TypeMismatch;
+		}
+
+		switch (Desc->ValueType)
+		{
+		case EPropertyBagPropertyType::Bool:
+		{
+			const FBoolProperty* Property = CastFieldChecked<FBoolProperty>(Desc->CachedProperty);
+			Property->SetPropertyValue(Address, InValue != 0);
+			return EPropertyBagResult::Success;
+		}
+		case EPropertyBagPropertyType::Byte:
+		{
+			const FByteProperty* Property = CastFieldChecked<FByteProperty>(Desc->CachedProperty);
+			Property->SetPropertyValue(Address, (uint8)InValue);
+			return EPropertyBagResult::Success;
+		}
+		case EPropertyBagPropertyType::Int32:
+		{
+			const FIntProperty* Property = CastFieldChecked<FIntProperty>(Desc->CachedProperty);
+			Property->SetPropertyValue(Address, (int32)InValue);
+			return EPropertyBagResult::Success;
+		}
+		case EPropertyBagPropertyType::UInt32:
+		{
+			const FUInt32Property* Property = CastFieldChecked<FUInt32Property>(Desc->CachedProperty);
+			Property->SetPropertyValue(Address, (uint32)InValue);
+			return EPropertyBagResult::Success;
+		}
+		case EPropertyBagPropertyType::Int64:
+		{
+			const FInt64Property* Property = CastFieldChecked<FInt64Property>(Desc->CachedProperty);
+			Property->SetPropertyValue(Address, (int64)InValue);
+			return EPropertyBagResult::Success;
+		}
+		case EPropertyBagPropertyType::UInt64:
+		{
+			const FUInt64Property* Property = CastFieldChecked<FUInt64Property>(Desc->CachedProperty);
+			Property->SetPropertyValue(Address, InValue);
+			return EPropertyBagResult::Success;
+		}
+		case EPropertyBagPropertyType::Float:
+		{
+			const FFloatProperty* Property = CastFieldChecked<FFloatProperty>(Desc->CachedProperty);
+			Property->SetPropertyValue(Address, (float)InValue);
+			return EPropertyBagResult::Success;
+		}
+		case EPropertyBagPropertyType::Double:
+		{
+			const FDoubleProperty* Property = CastFieldChecked<FDoubleProperty>(Desc->CachedProperty);
+			Property->SetPropertyValue(Address, (double)InValue);
+			return EPropertyBagResult::Success;
+		}
+		case EPropertyBagPropertyType::Enum:
+		{
+			const FEnumProperty* EnumProperty = CastFieldChecked<FEnumProperty>(Desc->CachedProperty);
+			const FNumericProperty* UnderlyingProperty = EnumProperty->GetUnderlyingProperty();
+			check(UnderlyingProperty);
+			UnderlyingProperty->SetIntPropertyValue(Address, InValue);
+			return EPropertyBagResult::Success;
+		}
+		default:
+			return EPropertyBagResult::TypeMismatch;
+		}
+	}
+
 	EPropertyBagResult SetPropertyFromDouble(const FPropertyBagPropertyDesc* Desc, void* Address, const double InValue)
 	{
 		if (Desc == nullptr || Desc->CachedProperty == nullptr)
@@ -794,12 +1006,24 @@ namespace UE::StructUtils::Private
 				Property->SetPropertyValue(Address, FMath::RoundToInt32(InValue));
 				return EPropertyBagResult::Success;
 			}
+		case EPropertyBagPropertyType::UInt32:
+			{
+				const FUInt32Property* Property = CastFieldChecked<FUInt32Property>(Desc->CachedProperty);
+				Property->SetPropertyValue(Address, FMath::RoundToInt32(InValue));
+				return EPropertyBagResult::Success;
+			}
 		case EPropertyBagPropertyType::Int64:
 			{
 				const FInt64Property* Property = CastFieldChecked<FInt64Property>(Desc->CachedProperty);
 				Property->SetPropertyValue(Address, FMath::RoundToInt64(InValue));
 				return EPropertyBagResult::Success;
 			}
+		case EPropertyBagPropertyType::UInt64:
+			{
+				const FUInt64Property* Property = CastFieldChecked<FUInt64Property>(Desc->CachedProperty);
+				Property->SetPropertyValue(Address, FMath::RoundToInt64(InValue));
+				return EPropertyBagResult::Success;
+			}
 		case EPropertyBagPropertyType::Float:
 			{
 				const FFloatProperty* Property = CastFieldChecked<FFloatProperty>(Desc->CachedProperty);
@@ -1062,10 +1286,21 @@ namespace UE::StructUtils::Private
 					}
 					else
 					{
-						int64 Value = 0;
-						if (GetPropertyAsInt64(&SourceDesc, SourceAddress, Value) == EPropertyBagResult::Success)
+						if (TargetDesc.IsUnsignedNumericType())
+						{
+							uint64 Value = 0;
+							if (GetPropertyAsUInt64(&SourceDesc, SourceAddress, Value) == EPropertyBagResult::Success)
+							{
+								SetPropertyFromUInt64(&TargetDesc, TargetAddress, Value);
+							}
+						}
+						else
 						{
-							SetPropertyFromInt64(&TargetDesc, TargetAddress, Value);
+							int64 Value = 0;
+							if (GetPropertyAsInt64(&SourceDesc, SourceAddress, Value) == EPropertyBagResult::Success)
+							{
+								SetPropertyFromInt64(&TargetDesc, TargetAddress, Value);
+							}
 						}
 					}
 				}
@@ -1236,13 +1471,26 @@ bool FPropertyBagPropertyDesc::IsNumericType() const
 	case EPropertyBagPropertyType::Bool: return true;
 	case EPropertyBagPropertyType::Byte: return true;
 	case EPropertyBagPropertyType::Int32: return true;
+	case EPropertyBagPropertyType::UInt32: return true;
 	case EPropertyBagPropertyType::Int64: return true;
+	case EPropertyBagPropertyType::UInt64: return true;
 	case EPropertyBagPropertyType::Float: return true;
 	case EPropertyBagPropertyType::Double: return true;
 	default: return false;
 	}
 }
 
+bool FPropertyBagPropertyDesc::IsUnsignedNumericType() const
+{
+	switch (ValueType)
+	{
+	case EPropertyBagPropertyType::Byte: return true;
+	case EPropertyBagPropertyType::UInt32: return true;
+	case EPropertyBagPropertyType::UInt64: return true;
+	default: return false;
+	}
+}
+
 bool FPropertyBagPropertyDesc::IsNumericFloatType() const
 {
 	switch (ValueType)
@@ -1470,8 +1718,8 @@ TValueOrError<bool, EPropertyBagResult> FInstancedPropertyBag::GetValueBool(cons
 TValueOrError<uint8, EPropertyBagResult> FInstancedPropertyBag::GetValueByte(const FName Name) const
 {
 	const FPropertyBagPropertyDesc* Desc = FindPropertyDescByName(Name);
-	int64 ReturnValue = 0;
-	const EPropertyBagResult Result = UE::StructUtils::Private::GetPropertyAsInt64(Desc, GetValueAddress(Desc), ReturnValue);
+	uint64 ReturnValue = 0;
+	const EPropertyBagResult Result = UE::StructUtils::Private::GetPropertyAsUInt64(Desc, GetValueAddress(Desc), ReturnValue);
 	if (Result != EPropertyBagResult::Success)
 	{
 		return MakeError(Result);
@@ -1491,6 +1739,18 @@ TValueOrError<int32, EPropertyBagResult> FInstancedPropertyBag::GetValueInt32(co
 	return MakeValue((int32)ReturnValue);
 }
 
+TValueOrError<uint32, EPropertyBagResult> FInstancedPropertyBag::GetValueUInt32(const FName Name) const
+{
+	const FPropertyBagPropertyDesc* Desc = FindPropertyDescByName(Name);
+	uint64 ReturnValue = 0;
+	const EPropertyBagResult Result = UE::StructUtils::Private::GetPropertyAsUInt64(Desc, GetValueAddress(Desc), ReturnValue);
+	if (Result != EPropertyBagResult::Success)
+	{
+		return MakeError(Result);
+	}
+	return MakeValue((uint32)ReturnValue);
+}
+
 TValueOrError<int64, EPropertyBagResult> FInstancedPropertyBag::GetValueInt64(const FName Name) const
 {
 	const FPropertyBagPropertyDesc* Desc = FindPropertyDescByName(Name);
@@ -1503,6 +1763,18 @@ TValueOrError<int64, EPropertyBagResult> FInstancedPropertyBag::GetValueInt64(co
 	return MakeValue(ReturnValue);
 }
 
+TValueOrError<uint64, EPropertyBagResult> FInstancedPropertyBag::GetValueUInt64(const FName Name) const
+{
+	const FPropertyBagPropertyDesc* Desc = FindPropertyDescByName(Name);
+	uint64 ReturnValue = 0;
+	const EPropertyBagResult Result = UE::StructUtils::Private::GetPropertyAsUInt64(Desc, GetValueAddress(Desc), ReturnValue);
+	if (Result != EPropertyBagResult::Success)
+	{
+		return MakeError(Result);
+	}
+	return MakeValue((uint64)ReturnValue);
+}
+
 TValueOrError<float, EPropertyBagResult> FInstancedPropertyBag::GetValueFloat(const FName Name) const
 {
 	const FPropertyBagPropertyDesc* Desc = FindPropertyDescByName(Name);
@@ -1660,7 +1932,7 @@ EPropertyBagResult FInstancedPropertyBag::SetValueBool(const FName Name, const b
 EPropertyBagResult FInstancedPropertyBag::SetValueByte(const FName Name, const uint8 InValue)
 {
 	const FPropertyBagPropertyDesc* Desc = FindPropertyDescByName(Name);
-	return UE::StructUtils::Private::SetPropertyFromInt64(Desc, GetMutableValueAddress(Desc), InValue);
+	return UE::StructUtils::Private::SetPropertyFromUInt64(Desc, GetMutableValueAddress(Desc), InValue);
 }
 
 EPropertyBagResult FInstancedPropertyBag::SetValueInt32(const FName Name, const int32 InValue)
@@ -1669,12 +1941,24 @@ EPropertyBagResult FInstancedPropertyBag::SetValueInt32(const FName Name, const
 	return UE::StructUtils::Private::SetPropertyFromInt64(Desc, GetMutableValueAddress(Desc), InValue);
 }
 
+EPropertyBagResult FInstancedPropertyBag::SetValueUInt32(const FName Name, const uint32 InValue)
+{
+	const FPropertyBagPropertyDesc* Desc = FindPropertyDescByName(Name);
+	return UE::StructUtils::Private::SetPropertyFromUInt64(Desc, GetMutableValueAddress(Desc), (int64)InValue);
+}
+
 EPropertyBagResult FInstancedPropertyBag::SetValueInt64(const FName Name, const int64 InValue)
 {
 	const FPropertyBagPropertyDesc* Desc = FindPropertyDescByName(Name);
 	return UE::StructUtils::Private::SetPropertyFromInt64(Desc, GetMutableValueAddress(Desc), InValue);
 }
 
+EPropertyBagResult FInstancedPropertyBag::SetValueUInt64(const FName Name, const uint64 InValue)
+{
+	const FPropertyBagPropertyDesc* Desc = FindPropertyDescByName(Name);
+	return UE::StructUtils::Private::SetPropertyFromUInt64(Desc, GetMutableValueAddress(Desc), (int64)InValue);
+}
+
 EPropertyBagResult FInstancedPropertyBag::SetValueFloat(const FName Name, const float InValue)
 {
 	const FPropertyBagPropertyDesc* Desc = FindPropertyDescByName(Name);
@@ -2056,8 +2340,8 @@ TValueOrError<bool, EPropertyBagResult> FPropertyBagArrayRef::GetValueBool(const
 
 TValueOrError<uint8, EPropertyBagResult> FPropertyBagArrayRef::GetValueByte(const int32 Index) const
 {
-	int64 IntValue = 0;
-	const EPropertyBagResult Result = UE::StructUtils::Private::GetPropertyAsInt64(&ValueDesc, GetAddress(Index), IntValue);
+	uint64 IntValue = 0;
+	const EPropertyBagResult Result = UE::StructUtils::Private::GetPropertyAsUInt64(&ValueDesc, GetAddress(Index), IntValue);
 	if (Result != EPropertyBagResult::Success)
 	{
 		return MakeError(Result);
@@ -2076,6 +2360,17 @@ TValueOrError<int32, EPropertyBagResult> FPropertyBagArrayRef::GetValueInt32(con
 	return MakeValue((int32)IntValue);
 }
 
+TValueOrError<uint32, EPropertyBagResult> FPropertyBagArrayRef::GetValueUInt32(const int32 Index) const
+{
+	uint64 IntValue = 0;
+	const EPropertyBagResult Result = UE::StructUtils::Private::GetPropertyAsUInt64(&ValueDesc, GetAddress(Index), IntValue);
+	if (Result != EPropertyBagResult::Success)
+	{
+		return MakeError(Result);
+	}
+	return MakeValue((uint32)IntValue);
+}
+
 TValueOrError<int64, EPropertyBagResult> FPropertyBagArrayRef::GetValueInt64(const int32 Index) const
 {
 	int64 IntValue = 0;
@@ -2087,6 +2382,17 @@ TValueOrError<int64, EPropertyBagResult> FPropertyBagArrayRef::GetValueInt64(con
 	return MakeValue(IntValue);
 }
 
+TValueOrError<uint64, EPropertyBagResult> FPropertyBagArrayRef::GetValueUInt64(const int32 Index) const
+{
+	uint64 IntValue = 0;
+	const EPropertyBagResult Result = UE::StructUtils::Private::GetPropertyAsUInt64(&ValueDesc, GetAddress(Index), IntValue);
+	if (Result != EPropertyBagResult::Success)
+	{
+		return MakeError(Result);
+	}
+	return MakeValue((uint64)IntValue);
+}
+
 TValueOrError<float, EPropertyBagResult> FPropertyBagArrayRef::GetValueFloat(const int32 Index) const
 {
 	double DblValue = 0;
@@ -2253,7 +2559,7 @@ EPropertyBagResult FPropertyBagArrayRef::SetValueBool(const int32 Index, const b
 
 EPropertyBagResult FPropertyBagArrayRef::SetValueByte(const int32 Index, const uint8 InValue)
 {
-	return UE::StructUtils::Private::SetPropertyFromInt64(&ValueDesc, GetMutableAddress(Index), InValue);
+	return UE::StructUtils::Private::SetPropertyFromUInt64(&ValueDesc, GetMutableAddress(Index), InValue);
 }
 
 EPropertyBagResult FPropertyBagArrayRef::SetValueInt32(const int32 Index, const int32 InValue)
@@ -2261,11 +2567,21 @@ EPropertyBagResult FPropertyBagArrayRef::SetValueInt32(const int32 Index, const
 	return UE::StructUtils::Private::SetPropertyFromInt64(&ValueDesc, GetMutableAddress(Index), InValue);
 }
 
+EPropertyBagResult FPropertyBagArrayRef::SetValueUInt32(const int32 Index, const uint32 InValue)
+{
+	return UE::StructUtils::Private::SetPropertyFromUInt64(&ValueDesc, GetMutableAddress(Index), InValue);
+}
+
 EPropertyBagResult FPropertyBagArrayRef::SetValueInt64(const int32 Index, const int64 InValue)
 {
 	return UE::StructUtils::Private::SetPropertyFromInt64(&ValueDesc, GetMutableAddress(Index), InValue);
 }
 
+EPropertyBagResult FPropertyBagArrayRef::SetValueUInt64(const int32 Index, const uint64 InValue)
+{
+	return UE::StructUtils::Private::SetPropertyFromUInt64(&ValueDesc, GetMutableAddress(Index), InValue);
+}
+
 EPropertyBagResult FPropertyBagArrayRef::SetValueFloat(const int32 Index, const float InValue)
 {
 	return UE::StructUtils::Private::SetPropertyFromDouble(&ValueDesc, GetMutableAddress(Index), InValue);
diff --git a/Engine/Plugins/Experimental/StructUtils/Source/StructUtils/Public/PropertyBag.h b/Engine/Plugins/Experimental/StructUtils/Source/StructUtils/Public/PropertyBag.h
--- a/Engine/Plugins/Experimental/StructUtils/Source/StructUtils/Public/PropertyBag.h
+++ b/Engine/Plugins/Experimental/StructUtils/Source/StructUtils/Public/PropertyBag.h
@@ -28,6 +28,8 @@ enum class EPropertyBagPropertyType : uint8
 	SoftObject UMETA(Hidden),
 	Class UMETA(Hidden),
 	SoftClass UMETA(Hidden),
+	UInt32,	// Type not fully supported at UI, will work with restrictions to type editing
+	UInt64, // Type not fully supported at UI, will work with restrictions to type editing
 
 	Count UMETA(Hidden)
 };
@@ -243,8 +245,11 @@ struct STRUCTUTILS_API FPropertyBagPropertyDesc
 	/** @return true if the two descriptors have the same type. Object types are compatible if Other can be cast to this type. */
 	bool CompatibleType(const FPropertyBagPropertyDesc& Other) const;
 
-	/** @return true if the property type is numeric (bool, int32, int64, float, double, enum) */
+	/** @return true if the property type is numeric (bool, (u)int32, (u)int64, float, double, enum) */
 	bool IsNumericType() const;
+
+	/** @return true if the property type is unsigned (uint32, uint64) */
+	bool IsUnsignedNumericType() const;
 	
 	/** @return true if the property type is floating point numeric (float, double) */
 	bool IsNumericFloatType() const;
@@ -448,13 +453,15 @@ struct STRUCTUTILS_API FInstancedPropertyBag
 	
 	/**
 	 * Getters
-	 * Numeric types (bool, int32, int64, float, double) support type conversion.
+	 * Numeric types (bool, (u)int32, (u)int64, float, double) support type conversion.
 	 */
 
 	TValueOrError<bool, EPropertyBagResult> GetValueBool(const FName Name) const;
 	TValueOrError<uint8, EPropertyBagResult> GetValueByte(const FName Name) const;
 	TValueOrError<int32, EPropertyBagResult> GetValueInt32(const FName Name) const;
+	TValueOrError<uint32, EPropertyBagResult> GetValueUInt32(const FName Name) const;
 	TValueOrError<int64, EPropertyBagResult> GetValueInt64(const FName Name) const;
+	TValueOrError<uint64, EPropertyBagResult> GetValueUInt64(const FName Name) const;
 	TValueOrError<float, EPropertyBagResult> GetValueFloat(const FName Name) const;
 	TValueOrError<double, EPropertyBagResult> GetValueDouble(const FName Name) const;
 	TValueOrError<FName, EPropertyBagResult> GetValueName(const FName Name) const;
@@ -523,12 +530,14 @@ struct STRUCTUTILS_API FInstancedPropertyBag
 
 	/**
 	 * Value Setters. A property must exists in that bag before it can be set.  
-	 * Numeric types (bool, int32, int64, float, double) support type conversion.
+	 * Numeric types (bool, (u)int32, (u)int64, float, double) support type conversion.
 	 */
 	EPropertyBagResult SetValueBool(const FName Name, const bool bInValue);
 	EPropertyBagResult SetValueByte(const FName Name, const uint8 InValue);
 	EPropertyBagResult SetValueInt32(const FName Name, const int32 InValue);
+	EPropertyBagResult SetValueUInt32(const FName Name, const uint32 InValue);
 	EPropertyBagResult SetValueInt64(const FName Name, const int64 InValue);
+	EPropertyBagResult SetValueUInt64(const FName Name, const uint64 InValue);
 	EPropertyBagResult SetValueFloat(const FName Name, const float InValue);
 	EPropertyBagResult SetValueDouble(const FName Name, const double InValue);
 	EPropertyBagResult SetValueName(const FName Name, const FName InValue);
@@ -675,13 +684,15 @@ public:
 
 	/**
 	 * Getters
-	 * Numeric types (bool, int32, int64, float, double) support type conversion.
+	 * Numeric types (bool, (u)int32, (u)int64, float, double) support type conversion.
 	 */
 	
 	TValueOrError<bool, EPropertyBagResult> GetValueBool(const int32 Index) const;
 	TValueOrError<uint8, EPropertyBagResult> GetValueByte(const int32 Index) const;
 	TValueOrError<int32, EPropertyBagResult> GetValueInt32(const int32 Index) const;
+	TValueOrError<uint32, EPropertyBagResult> GetValueUInt32(const int32 Index) const;
 	TValueOrError<int64, EPropertyBagResult> GetValueInt64(const int32 Index) const;
+	TValueOrError<uint64, EPropertyBagResult> GetValueUInt64(const int32 Index) const;
 	TValueOrError<float, EPropertyBagResult> GetValueFloat(const int32 Index) const;
 	TValueOrError<double, EPropertyBagResult> GetValueDouble(const int32 Index) const;
 	TValueOrError<FName, EPropertyBagResult> GetValueName(const int32 Index) const;
@@ -761,12 +772,14 @@ public:
 
 	/**
 	 * Value Setters. A property must exists in that bag before it can be set.  
-	 * Numeric types (bool, int32, int64, float, double) support type conversion.
+	 * Numeric types (bool, (u)int32, (u)int64, float, double) support type conversion.
 	 */
 	EPropertyBagResult SetValueBool(const int32 Index, const bool bInValue);
 	EPropertyBagResult SetValueByte(const int32 Index, const uint8 InValue);
 	EPropertyBagResult SetValueInt32(const int32 Index, const int32 InValue);
+	EPropertyBagResult SetValueUInt32(const int32 Index, const uint32 InValue);
 	EPropertyBagResult SetValueInt64(const int32 Index, const int64 InValue);
+	EPropertyBagResult SetValueUInt64(const int32 Index, const uint64 InValue);
 	EPropertyBagResult SetValueFloat(const int32 Index, const float InValue);
 	EPropertyBagResult SetValueDouble(const int32 Index, const double InValue);
 	EPropertyBagResult SetValueName(const int32 Index, const FName InValue);
diff --git a/Engine/Plugins/Experimental/StructUtils/Source/StructUtilsEditor/Private/Customizations/PropertyBagDetails.cpp b/Engine/Plugins/Experimental/StructUtils/Source/StructUtilsEditor/Private/Customizations/PropertyBagDetails.cpp
--- a/Engine/Plugins/Experimental/StructUtils/Source/StructUtilsEditor/Private/Customizations/PropertyBagDetails.cpp
+++ b/Engine/Plugins/Experimental/StructUtils/Source/StructUtilsEditor/Private/Customizations/PropertyBagDetails.cpp
@@ -288,6 +288,12 @@ FEdGraphPinType GetPropertyDescAsPin(const FPropertyBagPropertyDesc& Desc)
 		PinType.PinCategory = UEdGraphSchema_K2::PC_SoftClass;
 		PinType.PinSubCategoryObject = const_cast<UObject*>(Desc.ValueTypeObject.Get());
 		break;
+	case EPropertyBagPropertyType::UInt32:	// Warning : Type only partially supported (Blueprint does not support unsigned type)
+		PinType.PinCategory = UEdGraphSchema_K2::PC_Int;
+		break;
+	case EPropertyBagPropertyType::UInt64:	// Warning : Type only partially supported (Blueprint does not support unsigned type)
+		PinType.PinCategory = UEdGraphSchema_K2::PC_Int64;
+		break;
 	default:
 		ensureMsgf(false, TEXT("Unhandled value type %s"), *UEnum::GetValueAsString(Desc.ValueType));
 		break;
@@ -744,15 +750,31 @@ void FPropertyBagInstanceDataDetails::OnChildRowAdded(IDetailPropertyRow& ChildR
 	
 	bool bSupportedType = true;
 
+	FText UnsupportedTypeWarning;
+
 	TArray<FPropertyBagPropertyDesc> PropertyDescs = UE::StructUtils::Private::GetCommonPropertyDescs(BagStructProperty);
 	const FProperty* Property = ChildPropertyHandle->GetProperty();
 	if (FPropertyBagPropertyDesc* Desc = PropertyDescs.FindByPredicate([Property](const FPropertyBagPropertyDesc& Desc){ return Desc.CachedProperty == Property; }))
 	{
 		if (Desc->ContainerTypes.Num() > 1)
 		{
+			static const FText UnsupportedTypeWarningNestedContainer = LOCTEXT("NestedContainersWarning", "This property type is not supported in the property bag UI.");
+
+			// The property editing for nested containers is not supported.
+			UnsupportedTypeWarning = UnsupportedTypeWarningNestedContainer;
 			bSupportedType = false;
+			// Do not allow to edit the data.
 			ChildRow.IsEnabled(false);
 		}
+		else if ((Desc->ValueType == EPropertyBagPropertyType::UInt32 || Desc->ValueType == EPropertyBagPropertyType::UInt64)
+			&& !bFixedLayout)
+		{
+			static const FText UnsupportedTypeWarningUnsigned = LOCTEXT("UnsignedTypesWarning", "Unsigned types are not supported throught the property type selection. If you change the type, you will not be able to change it back.");
+
+			// Warn that the unsinged types cannot be set via the type selection.
+			UnsupportedTypeWarning = UnsupportedTypeWarningUnsigned;
+			bSupportedType = false;
+		}
 	}
 
 	if (!bFixedLayout)
@@ -829,10 +851,10 @@ void FPropertyBagInstanceDataDetails::OnChildRowAdded(IDetailPropertyRow& ChildR
 				.HeightOverride(12)
 				[
 					SNew(SImage)
-					.ToolTipText_Lambda([ChildPropertyHandle, bSupportedType]()
+					.ToolTipText_Lambda([UnsupportedTypeWarning, ChildPropertyHandle, bSupportedType]()
 					{
 						return !bSupportedType
-							? LOCTEXT("UnsupportedType", "This property type is not supported in the property bag UI.")
+							? UnsupportedTypeWarning
 							: LOCTEXT("MissingType", "The property is missing type. The Struct, Enum, or Object may have been removed.");
 					})
 					.Visibility_Lambda([ChildPropertyHandle, bSupportedType]()
diff --git a/Engine/Plugins/Experimental/StructUtils/Source/StructUtilsTestSuite/Private/PropertyBagTest.cpp b/Engine/Plugins/Experimental/StructUtils/Source/StructUtilsTestSuite/Private/PropertyBagTest.cpp
--- a/Engine/Plugins/Experimental/StructUtils/Source/StructUtilsTestSuite/Private/PropertyBagTest.cpp
+++ b/Engine/Plugins/Experimental/StructUtils/Source/StructUtilsTestSuite/Private/PropertyBagTest.cpp
@@ -15,6 +15,8 @@ struct FTest_CreatePropertyBag : FAITestBase
 		static const FName IsHotName(TEXT("bIsHot"));
 		static const FName TemperatureName(TEXT("Temperature"));
 		static const FName CountName(TEXT("Count"));
+		static const FName UInt32Name(TEXT("Unsigned32"));
+		static const FName UInt64Name(TEXT("Unsigned64"));
 
 		FInstancedPropertyBag Bag;
 
@@ -25,11 +27,22 @@ struct FTest_CreatePropertyBag : FAITestBase
 		// Amend the bag with new properties.
 		Bag.AddProperties({
 			{ TemperatureName, EPropertyBagPropertyType::Float },
-			{ CountName, EPropertyBagPropertyType::Int32 }
+			{ CountName, EPropertyBagPropertyType::Int32 },
+			{ UInt32Name, EPropertyBagPropertyType::UInt32 },
+			{ UInt64Name, EPropertyBagPropertyType::UInt64 }
 			});
 		AITEST_TRUE(TEXT("Set Temperature should succeed"), Bag.SetValueFloat(TemperatureName, 451.0f) == EPropertyBagResult::Success);
 		AITEST_TRUE(TEXT("Set Count should succeed"), Bag.SetValueFloat(CountName, 42) == EPropertyBagResult::Success);
 
+		AITEST_TRUE(TEXT("Set UInt32 should succeed"), Bag.SetValueUInt32(UInt32Name, MAX_uint32) == EPropertyBagResult::Success);
+		AITEST_TRUE(TEXT("Set UInt64 should succeed"), Bag.SetValueUInt64(UInt64Name, MAX_uint64) == EPropertyBagResult::Success);
+
+		AITEST_TRUE(TEXT("UInt32 value could not be retrieved"), Bag.GetValueUInt32(UInt32Name).HasError() == false);
+		AITEST_TRUE(TEXT("UInt32 value not correct"), Bag.GetValueUInt32(UInt32Name).GetValue() == MAX_uint32);
+
+		AITEST_TRUE(TEXT("UInt64 value could not be retrieved"), Bag.GetValueUInt64(UInt64Name).HasError() == false);
+		AITEST_TRUE(TEXT("UInt64 value not correct"), Bag.GetValueUInt64(UInt64Name).GetValue() == MAX_uint64);
+
 		Bag.RemovePropertyByName(IsHotName);
 		AITEST_TRUE(TEXT("Should not have bIsHot property"), Bag.FindPropertyDescByName(IsHotName) == nullptr);
 		AITEST_TRUE(TEXT("Set bIsHot should not succeed"), Bag.SetValueBool(IsHotName, true) != EPropertyBagResult::Success);
diff --git a/Engine/Plugins/MovieScene/MovieRenderPipeline/Source/MovieRenderPipelineCore/Public/Graph/MovieGraphCommon.h b/Engine/Plugins/MovieScene/MovieRenderPipeline/Source/MovieRenderPipelineCore/Public/Graph/MovieGraphCommon.h
--- a/Engine/Plugins/MovieScene/MovieRenderPipeline/Source/MovieRenderPipelineCore/Public/Graph/MovieGraphCommon.h
+++ b/Engine/Plugins/MovieScene/MovieRenderPipeline/Source/MovieRenderPipelineCore/Public/Graph/MovieGraphCommon.h
@@ -28,6 +28,8 @@ enum class EMovieGraphValueType : uint8
 	SoftObject UMETA(Hidden),
 	Class UMETA(Hidden),
 	SoftClass UMETA(Hidden),
+	UInt32,	// Type not fully supported at UI, will work with restrictions to type editing
+	UInt64, // Type not fully supported at UI, will work with restrictions to type editing
 
 	Count UMETA(Hidden)
 };
diff --git a/Engine/Plugins/Runtime/RigVM/Source/RigVM/Private/RigVMCore/RigVMMemoryStorageStruct.cpp b/Engine/Plugins/Runtime/RigVM/Source/RigVM/Private/RigVMCore/RigVMMemoryStorageStruct.cpp
--- a/Engine/Plugins/Runtime/RigVM/Source/RigVM/Private/RigVMCore/RigVMMemoryStorageStruct.cpp
+++ b/Engine/Plugins/Runtime/RigVM/Source/RigVM/Private/RigVMCore/RigVMMemoryStorageStruct.cpp
@@ -605,6 +605,8 @@ void FRigVMMemoryStorageStruct::SetDefaultValues(const TArray<FRigVMPropertyDesc
 	}
 
 	static const FName IntTypeName(TEXT("int")); // type used by some engine tests
+	static const FName Int64TypeName(TEXT("Int64"));
+	static const FName UInt64TypeName(TEXT("UInt64"));
 
 	if (VMType == RigVMTypeUtils::BoolTypeName)
 	{
@@ -614,6 +616,18 @@ void FRigVMMemoryStorageStruct::SetDefaultValues(const TArray<FRigVMPropertyDesc
 	{
 		OutBagPropertyType = EPropertyBagPropertyType::Int32;
 	}
+	else if (VMType == RigVMTypeUtils::UInt32TypeName)
+	{
+		OutBagPropertyType = EPropertyBagPropertyType::UInt32;
+	}
+	else if (VMType == Int64TypeName)
+	{
+		OutBagPropertyType = EPropertyBagPropertyType::Int64;
+	}
+	else if (VMType == UInt64TypeName)
+	{
+		OutBagPropertyType = EPropertyBagPropertyType::UInt64;
+	}
 	else if (VMType == RigVMTypeUtils::FloatTypeName)
 	{
 		OutBagPropertyType = EPropertyBagPropertyType::Float;
@@ -644,6 +658,7 @@ void FRigVMMemoryStorageStruct::SetDefaultValues(const TArray<FRigVMPropertyDesc
 	}
 	else
 	{
+		ensureMsgf(false, TEXT("Unsupported type : %s"), *VMType);
 		OutBagPropertyType = EPropertyBagPropertyType::None;
 	}
 
