UMG: Property bag should handle FClassProperty before FObjectProperty
- Property bag should preserve MetaClass for FClassProperty

diff --git a/Engine/Plugins/Experimental/StructUtils/Source/StructUtils/Private/PropertyBag.cpp b/Engine/Plugins/Experimental/StructUtils/Source/StructUtils/Private/PropertyBag.cpp
--- a/Engine/Plugins/Experimental/StructUtils/Source/StructUtils/Private/PropertyBag.cpp
+++ b/Engine/Plugins/Experimental/StructUtils/Source/StructUtils/Private/PropertyBag.cpp
@@ -155,6 +155,11 @@ namespace UE::StructUtils::Private
 		}
 		if (CastField<FObjectProperty>(InSourceProperty))
 		{
+			if (CastField<FClassProperty>(InSourceProperty))
+			{
+				return EPropertyBagPropertyType::Class;
+			}
+
 			return EPropertyBagPropertyType::Object;
 		}
 		if (CastField<FObjectPtrProperty>(InSourceProperty))
@@ -163,16 +168,13 @@ namespace UE::StructUtils::Private
 		}
 		if (CastField<FSoftObjectProperty>(InSourceProperty))
 		{
+			if (CastField<FSoftClassProperty>(InSourceProperty))
+			{
+				return EPropertyBagPropertyType::SoftClass;
+			}
+
 			return EPropertyBagPropertyType::SoftObject;
 		}
-		if (CastField<FClassProperty>(InSourceProperty))
-		{
-			return EPropertyBagPropertyType::Class;
-		}
-		if (CastField<FSoftClassProperty>(InSourceProperty))
-		{
-			return EPropertyBagPropertyType::SoftClass;
-		}
 
 		// Handle array property
 		if (const FArrayProperty* ArrayProperty = CastField<FArrayProperty>(InSourceProperty))
@@ -202,6 +204,11 @@ namespace UE::StructUtils::Private
 		}
 		if (const FObjectProperty* ObjectProperty = CastField<FObjectProperty>(InSourceProperty))
 		{
+			if (const FClassProperty* ClassProperty = CastField<FClassProperty>(InSourceProperty))
+			{
+				return ClassProperty->PropertyClass;
+			}
+
 			return ObjectProperty->PropertyClass;
 		}
 		if (const FObjectPtrProperty* ObjectProperty = CastField<FObjectPtrProperty>(InSourceProperty))
@@ -210,16 +217,13 @@ namespace UE::StructUtils::Private
 		}
 		if (const FSoftObjectProperty* SoftObjectProperty = CastField<FSoftObjectProperty>(InSourceProperty))
 		{
+			if (const FSoftClassProperty* SoftClassProperty = CastField<FSoftClassProperty>(InSourceProperty))
+			{
+				return SoftClassProperty->PropertyClass;
+			}
+
 			return SoftObjectProperty->PropertyClass;
 		}
-		if (const FClassProperty* ClassProperty = CastField<FClassProperty>(InSourceProperty))
-		{
-			return ClassProperty->PropertyClass;
-		}
-		if (const FSoftClassProperty* SoftClassProperty = CastField<FSoftClassProperty>(InSourceProperty))
-		{
-			return SoftClassProperty->PropertyClass;
-		}
 		
 		// Handle array property
 		if (const FArrayProperty* ArrayProperty = CastField<FArrayProperty>(InSourceProperty))
@@ -227,7 +231,6 @@ namespace UE::StructUtils::Private
 			return GetValueTypeObjectFromProperty(ArrayProperty->Inner);
 		}
 
-
 		return nullptr;
 	}
 
@@ -370,6 +373,34 @@ namespace UE::StructUtils::Private
 				return Prop;
 			}
 			break;
+		case EPropertyBagPropertyType::Class:
+			if (const UClass* Class = Cast<UClass>(Desc.ValueTypeObject))
+			{
+				FClassProperty* Prop = new FClassProperty(PropertyScope, Desc.Name, RF_Public);
+#if WITH_EDITORONLY_DATA
+				Prop->SetMetaClass(Desc.MetaClass ? Desc.MetaClass.Get() : const_cast<UClass*>(Class));
+#else
+				Prop->SetMetaClass(const_cast<UClass*>(Class));
+#endif
+				Prop->PropertyClass = UClass::StaticClass();
+				Prop->SetPropertyFlags(CPF_HasGetValueTypeHash);
+				return Prop;
+			}
+			break;
+		case EPropertyBagPropertyType::SoftClass:
+			if (const UClass* Class = Cast<UClass>(Desc.ValueTypeObject))
+			{
+				FSoftClassProperty* Prop = new FSoftClassProperty(PropertyScope, Desc.Name, RF_Public);
+#if WITH_EDITORONLY_DATA
+				Prop->SetMetaClass(Desc.MetaClass ? Desc.MetaClass.Get() : const_cast<UClass*>(Class));
+#else
+				Prop->SetMetaClass(const_cast<UClass*>(Class));
+#endif
+				Prop->PropertyClass = UClass::StaticClass();
+				Prop->SetPropertyFlags(CPF_HasGetValueTypeHash);
+				return Prop;
+			}
+			break;
 		case EPropertyBagPropertyType::Object:
 			if (const UClass* Class = Cast<UClass>(Desc.ValueTypeObject))
 			{
@@ -396,26 +427,6 @@ namespace UE::StructUtils::Private
 				return Prop;
 			}
 			break;
-		case EPropertyBagPropertyType::Class:
-			if (const UClass* Class = Cast<UClass>(Desc.ValueTypeObject))
-			{
-				FClassProperty* Prop = new FClassProperty(PropertyScope, Desc.Name, RF_Public);
-				Prop->SetMetaClass(const_cast<UClass*>(Class));
-				Prop->PropertyClass = UClass::StaticClass();
-				Prop->SetPropertyFlags(CPF_HasGetValueTypeHash);
-				return Prop;
-			}
-			break;
-		case EPropertyBagPropertyType::SoftClass:
-			if (const UClass* Class = Cast<UClass>(Desc.ValueTypeObject))
-			{
-				FSoftClassProperty* Prop = new FSoftClassProperty(PropertyScope, Desc.Name, RF_Public);
-				Prop->SetMetaClass(const_cast<UClass*>(Class));
-				Prop->PropertyClass = UClass::StaticClass();
-				Prop->SetPropertyFlags(CPF_HasGetValueTypeHash);
-				return Prop;
-			}
-			break;
 		default:
 			ensureMsgf(false, TEXT("Unhandled stype %s"), *UEnum::GetValueAsString(Desc.ValueType));
 		}
@@ -1415,6 +1426,11 @@ FPropertyBagPropertyDesc::FPropertyBagPropertyDesc(const FName InName, const FPr
 	ContainerTypes = UE::StructUtils::Private::GetContainerTypesFromProperty(InSourceProperty);
 
 #if WITH_EDITORONLY_DATA
+	if (const FClassProperty* ClassProperty = CastFieldChecked<FClassProperty>(InSourceProperty))
+	{
+		MetaClass = ClassProperty->MetaClass;
+	}
+
 	if (const TMap<FName, FString>* SourcePropertyMetaData = InSourceProperty->GetMetaDataMap())
 	{
 		for (const TPair<FName, FString>& MetaDataPair : *SourcePropertyMetaData)
diff --git a/Engine/Plugins/Experimental/StructUtils/Source/StructUtils/Public/PropertyBag.h b/Engine/Plugins/Experimental/StructUtils/Source/StructUtils/Public/PropertyBag.h
--- a/Engine/Plugins/Experimental/StructUtils/Source/StructUtils/Public/PropertyBag.h
+++ b/Engine/Plugins/Experimental/StructUtils/Source/StructUtils/Public/PropertyBag.h
@@ -284,6 +284,10 @@ struct STRUCTUTILS_API FPropertyBagPropertyDesc
 	/** Editor-only meta data for CachedProperty */
 	UPROPERTY(EditAnywhere, Category="Default")
 	TArray<FPropertyBagPropertyDescMetaData> MetaData;
+
+	/** Editor-only meta class for IClassViewer */
+	UPROPERTY(EditAnywhere, Category = "Default")
+	TObjectPtr<class UClass> MetaClass;
 #endif
 
 	/** Cached property pointer, set in UPropertyBag::GetOrCreateFromDescs. */
