Fix CIS error using CastFieldChecked in FPropertyBagPropertyDesc constructor
* Fix is to use the non-checked version since InSourceProperty is not guaranteed to be an FClassProperty.
* Original issue introduced in "UMG: Property bag should handle FClassProperty before FObjectProperty..."

diff --git a/Engine/Plugins/Experimental/StructUtils/Source/StructUtils/Private/PropertyBag.cpp b/Engine/Plugins/Experimental/StructUtils/Source/StructUtils/Private/PropertyBag.cpp
--- a/Engine/Plugins/Experimental/StructUtils/Source/StructUtils/Private/PropertyBag.cpp
+++ b/Engine/Plugins/Experimental/StructUtils/Source/StructUtils/Private/PropertyBag.cpp
@@ -1426,7 +1426,7 @@ FPropertyBagPropertyDesc::FPropertyBagPropertyDesc(const FName InName, const FPr
 	ContainerTypes = UE::StructUtils::Private::GetContainerTypesFromProperty(InSourceProperty);
 
 #if WITH_EDITORONLY_DATA
-	if (const FClassProperty* ClassProperty = CastFieldChecked<FClassProperty>(InSourceProperty))
+	if (const FClassProperty* ClassProperty = CastField<FClassProperty>(InSourceProperty))
 	{
 		MetaClass = ClassProperty->MetaClass;
 	}
