UMG: Serialize MetaClass property in PropertyBag

diff --git a/Engine/Plugins/Experimental/StructUtils/Source/StructUtils/Private/PropertyBag.cpp b/Engine/Plugins/Experimental/StructUtils/Source/StructUtils/Private/PropertyBag.cpp
--- a/Engine/Plugins/Experimental/StructUtils/Source/StructUtils/Private/PropertyBag.cpp
+++ b/Engine/Plugins/Experimental/StructUtils/Source/StructUtils/Private/PropertyBag.cpp
@@ -25,6 +25,7 @@ struct STRUCTUTILS_API FPropertyBagCustomVersion
 		// Added support for array types
 		ContainerTypes = 1,
 		NestedContainerTypes = 2,
+		MetaClass = 3,
 
 		// -----<new versions can be added above this line>-------------------------------------------------
 		VersionPlusOne,
@@ -1468,9 +1469,9 @@ static FArchive& operator<<(FArchive& Ar, FPropertyBagPropertyDesc& Bag)
 
 	bool bHasMetaData = false;
 #if WITH_EDITORONLY_DATA
-	if (Ar.IsSaving())
+	if (Ar.IsSaving() && !Ar.IsCooking())
 	{
-		bHasMetaData = !Ar.IsCooking() && Bag.MetaData.Num() > 0;
+		bHasMetaData = Bag.MetaData.Num() > 0 || Bag.MetaClass;
 	}
 #endif
 	Ar << bHasMetaData;
@@ -1479,11 +1480,22 @@ static FArchive& operator<<(FArchive& Ar, FPropertyBagPropertyDesc& Bag)
 	{
 #if WITH_EDITORONLY_DATA
 		Ar << Bag.MetaData;
+
+		if (Ar.CustomVer(FPropertyBagCustomVersion::GUID) >= FPropertyBagCustomVersion::MetaClass)
+		{
+			Ar << Bag.MetaClass;
+		}
 #else
 		TArray<FPropertyBagPropertyDescMetaData> TempMetaData; 
 		Ar << TempMetaData;
+
+		if (Ar.CustomVer(FPropertyBagCustomVersion::GUID) >= FPropertyBagCustomVersion::MetaClass)
+		{
+			TObjectPtr<UClass> TempMetaClass = nullptr;
+			Ar << TempMetaClass;
+		}
 #endif
-	}	
+	}
 	
 	return Ar;
 }
