StructUtils: Prevent crash when calling IsChildOf() on the property bag structs on non-editor builds

diff --git a/Engine/Plugins/Experimental/StructUtils/Source/StructUtils/Private/PropertyBag.cpp b/Engine/Plugins/Experimental/StructUtils/Source/StructUtils/Private/PropertyBag.cpp
--- a/Engine/Plugins/Experimental/StructUtils/Source/StructUtils/Private/PropertyBag.cpp
+++ b/Engine/Plugins/Experimental/StructUtils/Source/StructUtils/Private/PropertyBag.cpp
@@ -2750,6 +2750,13 @@ const UPropertyBag* UPropertyBag::GetOrCreateFromDescs(const TConstArrayView<FPr
 		}
 	}
 
+	// @hack:
+	// This method is called to prevent non-editor builds to not crash on IsChildOf().
+	// The issues is that the UScriptStruct(const FObjectInitializer& ObjectInitializer) ctor (which is macro/UHT generated)
+	// does not call ReinitializeBaseChainArray(), when the code is compiled with USTRUCT_ISCHILDOF_STRUCTARRAY.
+	// Calling SetSuperStruct() forces the ReinitializeBaseChainArray() to be called.
+	NewBag->SetSuperStruct(nullptr);
+	
 	NewBag->Bind();
 	NewBag->StaticLink(/*RelinkExistingProperties*/true);
 
