Fix crashes when creating property bags with identical layouts from multiple threads

diff --git a/Engine/Plugins/Experimental/StructUtils/Source/StructUtils/Private/PropertyBag.cpp b/Engine/Plugins/Experimental/StructUtils/Source/StructUtils/Private/PropertyBag.cpp
--- a/Engine/Plugins/Experimental/StructUtils/Source/StructUtils/Private/PropertyBag.cpp
+++ b/Engine/Plugins/Experimental/StructUtils/Source/StructUtils/Private/PropertyBag.cpp
@@ -2664,13 +2664,23 @@ EPropertyBagResult FPropertyBagArrayRef::SetValueSoftPath(const int32 Index, con
 //  UPropertyBag
 //----------------------------------------------------------------//
 
+namespace UE::StructUtils::Private
+{
+	// Lock to prevent concurrent access to lazily-constructed UPropertyBag objects in UPropertyBag::GetOrCreateFromDescs
+	static FCriticalSection GPropertyBagLock;
+}
+
 const UPropertyBag* UPropertyBag::GetOrCreateFromDescs(const TConstArrayView<FPropertyBagPropertyDesc> PropertyDescs, const TCHAR* PrefixName)
 {
 	const uint64 BagHash = UE::StructUtils::Private::CalcPropertyDescArrayHash(PropertyDescs);
 	const FString ScriptStructName = PrefixName == nullptr
 		? FString::Printf(TEXT("PropertyBag_%llx"), BagHash)
 		: FString::Printf(TEXT("%s_%llx"), PrefixName, BagHash);
-	
+
+	// We need to linearize this entire function otherwise threads that create bags of identical layouts can view
+	// partially-constructed objects 
+	FScopeLock ScopeLock(&UE::StructUtils::Private::GPropertyBagLock);
+
 	if (const UPropertyBag* ExistingBag = FindObject<UPropertyBag>(GetTransientPackage(), *ScriptStructName))
 	{
 		return ExistingBag;
