StateTree: Added option for State Tree parameters to be overridden.
- Added property override logic into FPropertyBagInstanceDataDetails
- Added hooks to allow to customize the peoperty override logic by deriving from FPropertyBagInstanceDataDetails
- Added array of overridden property IDs in FStateTreeReference

diff --git a/Engine/Plugins/Experimental/StructUtils/Source/StructUtils/Private/PropertyBag.cpp b/Engine/Plugins/Experimental/StructUtils/Source/StructUtils/Private/PropertyBag.cpp
--- a/Engine/Plugins/Experimental/StructUtils/Source/StructUtils/Private/PropertyBag.cpp
+++ b/Engine/Plugins/Experimental/StructUtils/Source/StructUtils/Private/PropertyBag.cpp
@@ -1248,7 +1248,13 @@ namespace UE::StructUtils::Private
 		return EPropertyBagResult::Success;
 	}
 
-	void CopyMatchingValuesByID(const FConstStructView Source, FStructView Target)
+	/**
+	 * Copies properties from Source to Target property bag. The bag layouts does not need to match.
+	 * Properties are matched based on the ID in the property bag descs.
+	 * If bHasOverrides == true, then only the matching properties, whose ID is found in Overrides are copied (if Overrides is empty, nothing is copied).
+	 * If bHasOverrides == false, all matching properties are copied.
+	 */
+	void CopyMatchingValuesByID(const FConstStructView Source, FStructView Target, bool bHasOverrides, TConstArrayView<FGuid> Overrides)
 	{
 		if (!Source.IsValid() || !Target.IsValid())
 		{
@@ -1266,6 +1272,12 @@ namespace UE::StructUtils::Private
 		// Iterate over source and copy to target if possible. Source is expected to usually have less items.
 		for (const FPropertyBagPropertyDesc& SourceDesc : SourceBagStruct->GetPropertyDescs())
 		{
+			const bool bShouldCopy = !bHasOverrides || Overrides.Contains(SourceDesc.ID);
+			if (!bShouldCopy)
+			{
+				continue;
+			}
+			
 			const FPropertyBagPropertyDesc* PotentialTargetDesc = TargetBagStruct->FindPropertyDescByID(SourceDesc.ID);
 			if (PotentialTargetDesc == nullptr
 				|| PotentialTargetDesc->CachedProperty == nullptr
@@ -1333,6 +1345,16 @@ namespace UE::StructUtils::Private
 		}
 	}
 
+	void CopyMatchingValuesByID(const FConstStructView Source, FStructView Target)
+	{
+		CopyMatchingValuesByID(Source, Target, /*bHasOverrides*/false, {});
+	}
+
+	void CopyMatchingValuesByIDWithOverrides(const FConstStructView Source, FStructView Target, TConstArrayView<FGuid> Overrides)
+	{
+		CopyMatchingValuesByID(Source, Target, /*bHasOverrides*/true, Overrides);
+	}
+
 	void RemovePropertyByName(TArray<FPropertyBagPropertyDesc>& Descs, const FName PropertyName, const int32 StartIndex = 0)
 	{
 		// Remove properties which dont have unique name.
@@ -1690,6 +1712,15 @@ void FInstancedPropertyBag::MigrateToNewBagInstance(const FInstancedPropertyBag&
 	Value = MoveTemp(NewValue);
 }
 
+void FInstancedPropertyBag::MigrateToNewBagInstanceWithOverrides(const FInstancedPropertyBag& NewBagInstance, TConstArrayView<FGuid> OverriddenPropertyIDs)
+{
+	FInstancedStruct NewValue(NewBagInstance.Value);
+
+	UE::StructUtils::Private::CopyMatchingValuesByIDWithOverrides(Value, NewValue, OverriddenPropertyIDs);
+	
+	Value = MoveTemp(NewValue);
+}
+
 const UPropertyBag* FInstancedPropertyBag::GetPropertyBagStruct() const
 {
 	return Value.IsValid() ? Cast<const UPropertyBag>(Value.GetScriptStruct()) : nullptr;
diff --git a/Engine/Plugins/Experimental/StructUtils/Source/StructUtils/Public/PropertyBag.h b/Engine/Plugins/Experimental/StructUtils/Source/StructUtils/Public/PropertyBag.h
--- a/Engine/Plugins/Experimental/StructUtils/Source/StructUtils/Public/PropertyBag.h
+++ b/Engine/Plugins/Experimental/StructUtils/Source/StructUtils/Public/PropertyBag.h
@@ -440,6 +440,15 @@ struct STRUCTUTILS_API FInstancedPropertyBag
 	 */
 	void MigrateToNewBagInstance(const FInstancedPropertyBag& NewBagInstance);
 
+	/**
+	 * Changes the type of this bag to the specified other bag, copies base values from the other bag,
+	 * and migrates existing values over if marked as overridden in the OverriddenPropertyIDs.
+	 * The properties are matched between the bags based on the property ID.
+	 * @param NewBagInstance Reference to the new type.
+	 * @param OverriddenPropertyIDs Array if property IDs which should be copied over to the new instance. 
+	 */
+	void MigrateToNewBagInstanceWithOverrides(const FInstancedPropertyBag& NewBagInstance, TConstArrayView<FGuid> OverriddenPropertyIDs);
+
 	/** @return pointer to the property bag struct. */ 
 	const UPropertyBag* GetPropertyBagStruct() const;
 	
diff --git a/Engine/Plugins/Experimental/StructUtils/Source/StructUtilsEditor/Private/Customizations/PropertyBagDetails.cpp b/Engine/Plugins/Experimental/StructUtils/Source/StructUtilsEditor/Private/Customizations/PropertyBagDetails.cpp
--- a/Engine/Plugins/Experimental/StructUtils/Source/StructUtilsEditor/Private/Customizations/PropertyBagDetails.cpp
+++ b/Engine/Plugins/Experimental/StructUtils/Source/StructUtilsEditor/Private/Customizations/PropertyBagDetails.cpp
@@ -719,6 +719,54 @@ TValueOrError<ReturnType, void> CallFunc(UObject* InTargetObject, UFunction* InF
 	}
 }
 
+/** Checks if the value for a source property in a source struct has the same value that the target property in the target struct. */
+bool ArePropertiesIdentical(
+	const FPropertyBagPropertyDesc* InSourcePropertyDesc,
+	const FInstancedPropertyBag& InSourceInstance,
+	const FPropertyBagPropertyDesc* InTargetPropertyDesc,
+	const FInstancedPropertyBag& InTargetInstance)
+{
+	if (!InSourceInstance.IsValid()
+		|| !InTargetInstance.IsValid()
+		|| !InSourcePropertyDesc
+		|| !InSourcePropertyDesc->CachedProperty
+		|| !InTargetPropertyDesc
+		|| !InTargetPropertyDesc->CachedProperty)
+	{
+		return false;
+	}
+
+	if (!InSourcePropertyDesc->CompatibleType(*InTargetPropertyDesc))
+	{
+		return false;
+	}
+
+	const uint8* SourceValueAddress = InSourceInstance.GetValue().GetMemory() + InSourcePropertyDesc->CachedProperty->GetOffset_ForInternal();
+	const uint8* TargetValueAddress = InTargetInstance.GetValue().GetMemory() + InTargetPropertyDesc->CachedProperty->GetOffset_ForInternal();
+
+	return InSourcePropertyDesc->CachedProperty->Identical(SourceValueAddress, TargetValueAddress);
+}
+
+/** Copy the value for a source property in a source struct to the target property in the target struct. */
+void CopyPropertyValue(const FPropertyBagPropertyDesc* InSourcePropertyDesc, const FInstancedPropertyBag& InSourceInstance, const FPropertyBagPropertyDesc* InTargetPropertyDesc, FInstancedPropertyBag& InTargetInstance)
+{
+	if (!InSourceInstance.IsValid() || !InTargetInstance.IsValid() || !InSourcePropertyDesc || !InSourcePropertyDesc->CachedProperty || !InTargetPropertyDesc || !InTargetPropertyDesc->CachedProperty)
+	{
+		return;
+	}
+
+	// Can't copy if they are not compatible.
+	if (!InSourcePropertyDesc->CompatibleType(*InTargetPropertyDesc))
+	{
+		return;
+	}
+
+	const uint8* SourceValueAddress = InSourceInstance.GetValue().GetMemory() + InSourcePropertyDesc->CachedProperty->GetOffset_ForInternal();
+	uint8* TargetValueAddress = InTargetInstance.GetMutableValue().GetMemory() + InTargetPropertyDesc->CachedProperty->GetOffset_ForInternal();
+
+	InSourcePropertyDesc->CachedProperty->CopyCompleteValue(TargetValueAddress, SourceValueAddress);
+}
+
 } // UE::StructUtils::Private
 
 
@@ -858,7 +906,7 @@ void FPropertyBagInstanceDataDetails::OnChildRowAdded(IDetailPropertyRow& ChildR
 				.HeightOverride(12)
 				[
 					SNew(SImage)
-					.ToolTipText_Lambda([UnsupportedTypeWarning, ChildPropertyHandle, bSupportedType]()
+					.ToolTipText_Lambda([UnsupportedTypeWarning, bSupportedType]()
 					{
 						return !bSupportedType
 							? UnsupportedTypeWarning
@@ -882,6 +930,212 @@ void FPropertyBagInstanceDataDetails::OnChildRowAdded(IDetailPropertyRow& ChildR
 		[
 			ValueWidget.ToSharedRef()
 		];
+
+
+	if (HasPropertyOverrides())
+	{
+		TWeakPtr<FPropertyBagInstanceDataDetails> WeakSelf = SharedThis<FPropertyBagInstanceDataDetails>(this);
+
+		TAttribute<bool> EditConditionValue = TAttribute<bool>::CreateLambda(
+			[WeakSelf, ChildPropertyHandle]() -> bool
+			{
+				if (const TSharedPtr<FPropertyBagInstanceDataDetails> Self = WeakSelf.Pin())
+				{
+					return Self->IsPropertyOverridden(ChildPropertyHandle) == EPropertyOverrideState::Yes;
+				}
+				return true;
+			});
+		
+		FOnBooleanValueChanged OnEditConditionChanged = FOnBooleanValueChanged::CreateLambda([WeakSelf, ChildPropertyHandle](bool bNewValue)
+		{
+			if (const TSharedPtr<FPropertyBagInstanceDataDetails> Self = WeakSelf.Pin())
+			{
+				Self->SetPropertyOverride(ChildPropertyHandle, bNewValue);
+			}
+		});
+
+		ChildRow.EditCondition(std::move(EditConditionValue), std::move(OnEditConditionChanged));
+
+		FIsResetToDefaultVisible IsResetVisible = FIsResetToDefaultVisible::CreateLambda([WeakSelf](TSharedPtr<IPropertyHandle> PropertyHandle)
+		{
+			if (const TSharedPtr<FPropertyBagInstanceDataDetails> Self = WeakSelf.Pin())
+			{
+				return !Self->IsDefaultValue(PropertyHandle);
+			}
+			return false;
+		});
+		FResetToDefaultHandler ResetHandler = FResetToDefaultHandler::CreateLambda([WeakSelf](TSharedPtr<IPropertyHandle> PropertyHandle)
+		{
+			if (const TSharedPtr<FPropertyBagInstanceDataDetails> Self = WeakSelf.Pin())
+			{
+				Self->ResetToDefault(PropertyHandle);
+			}
+		});
+		FResetToDefaultOverride ResetOverride = FResetToDefaultOverride::Create(IsResetVisible, ResetHandler);
+
+		ChildRow.OverrideResetToDefault(ResetOverride);
+	}
+}
+
+FPropertyBagInstanceDataDetails::EPropertyOverrideState FPropertyBagInstanceDataDetails::IsPropertyOverridden(TSharedPtr<IPropertyHandle> ChildPropertyHandle) const
+{
+	if (!ChildPropertyHandle)
+	{
+		return EPropertyOverrideState::Undetermined;;
+	}
+
+	int32 NumValues = 0;
+	int32 NumOverrides = 0; 
+
+	const FProperty* Property = ChildPropertyHandle->GetProperty();
+	check(Property);
+
+	EnumeratePropertyBags(BagStructProperty,
+		[Property, &NumValues, &NumOverrides]
+		(const FInstancedPropertyBag& DefaultPropertyBag, const FInstancedPropertyBag& PropertyBag, const IPropertyBagOverrideProvider& OverrideProvider)
+		{
+			NumValues++;
+			if (const UPropertyBag* Bag = PropertyBag.GetPropertyBagStruct())
+			{
+				const FPropertyBagPropertyDesc* PropertyDesc = Bag->FindPropertyDescByPropertyName(Property->GetFName());
+				if (PropertyDesc && OverrideProvider.IsPropertyOverridden(PropertyDesc->ID))
+				{
+					NumOverrides++;
+				}
+			}
+
+			return true;
+		});
+	
+	if (NumOverrides == 0)
+	{
+		return EPropertyOverrideState::No;
+	}
+	else if (NumOverrides == NumValues)
+	{
+		return EPropertyOverrideState::Yes;
+	}			
+	return EPropertyOverrideState::Undetermined;
+}
+	
+void FPropertyBagInstanceDataDetails::SetPropertyOverride(TSharedPtr<IPropertyHandle> ChildPropertyHandle, const bool bIsOverridden)
+{
+	if (!ChildPropertyHandle)
+	{
+		return;
+	}
+
+	const FProperty* Property = ChildPropertyHandle->GetProperty();
+	check(Property);
+
+	FScopedTransaction Transaction(FText::Format(LOCTEXT("OverrideChange", "Change Override for {0}"), FText::FromName(ChildPropertyHandle->GetProperty()->GetFName())));
+	
+	PreChangeOverrides();
+	
+	EnumeratePropertyBags(
+		BagStructProperty,
+		[Property, bIsOverridden]
+		(const FInstancedPropertyBag& DefaultPropertyBag, const FInstancedPropertyBag& PropertyBag, const IPropertyBagOverrideProvider& OverrideProvider)
+		{
+			if (const UPropertyBag* Bag = PropertyBag.GetPropertyBagStruct())
+			{
+				if (const FPropertyBagPropertyDesc* PropertyDesc = Bag->FindPropertyDescByPropertyName(Property->GetFName()))
+				{
+					OverrideProvider.SetPropertyOverride(PropertyDesc->ID, bIsOverridden);
+				}
+			}
+
+			return true;
+		});
+
+	PostChangeOverrides();
+}
+
+bool FPropertyBagInstanceDataDetails::IsDefaultValue(TSharedPtr<IPropertyHandle> ChildPropertyHandle) const
+{
+	if (!ChildPropertyHandle)
+	{
+		return true;
+	}
+
+	int32 NumValues = 0;
+	int32 NumOverridden = 0;
+	int32 NumIdentical = 0;
+
+	const FProperty* Property = ChildPropertyHandle->GetProperty();
+	check(Property);
+
+	EnumeratePropertyBags(
+		BagStructProperty,
+		[Property, &NumValues, &NumOverridden, &NumIdentical]
+		(const FInstancedPropertyBag& DefaultPropertyBag, const FInstancedPropertyBag& PropertyBag, const IPropertyBagOverrideProvider& OverrideProvider)
+		{
+			NumValues++;
+
+			const UPropertyBag* DefaultBag = DefaultPropertyBag.GetPropertyBagStruct();
+			const UPropertyBag* Bag = PropertyBag.GetPropertyBagStruct();
+			if (Bag && DefaultBag)
+			{
+				const FPropertyBagPropertyDesc* PropertyDesc = Bag->FindPropertyDescByName(Property->GetFName());
+				const FPropertyBagPropertyDesc* DefaultPropertyDesc = DefaultBag->FindPropertyDescByName(Property->GetFName());
+				if (PropertyDesc
+					&& DefaultPropertyDesc
+					&& OverrideProvider.IsPropertyOverridden(PropertyDesc->ID))
+				{
+					NumOverridden++;
+					if (UE::StructUtils::Private::ArePropertiesIdentical(DefaultPropertyDesc, DefaultPropertyBag, PropertyDesc, PropertyBag))
+					{
+						NumIdentical++;
+					}
+				}
+			}
+			return true;
+		});
+
+	if (NumOverridden == NumIdentical)
+	{
+		return true;
+	}
+	
+	return false;
+}
+
+void FPropertyBagInstanceDataDetails::ResetToDefault(TSharedPtr<IPropertyHandle> ChildPropertyHandle)
+{
+	if (!ChildPropertyHandle)
+	{
+		return;
+	}
+	
+	const FProperty* Property = ChildPropertyHandle->GetProperty();
+	check(Property);
+
+	FScopedTransaction Transaction(FText::Format(LOCTEXT("ResetToDefault", "Reset {0} to default value"), FText::FromName(ChildPropertyHandle->GetProperty()->GetFName())));
+	ChildPropertyHandle->NotifyPreChange();
+	
+	EnumeratePropertyBags(
+		BagStructProperty,
+		[Property]
+		(const FInstancedPropertyBag& DefaultPropertyBag, FInstancedPropertyBag& PropertyBag, const IPropertyBagOverrideProvider& OverrideProvider)
+		{
+			const UPropertyBag* DefaultBag = DefaultPropertyBag.GetPropertyBagStruct();
+			const UPropertyBag* Bag = PropertyBag.GetPropertyBagStruct();
+			if (Bag && DefaultBag)
+			{
+				const FPropertyBagPropertyDesc* PropertyDesc = Bag->FindPropertyDescByName(Property->GetFName());
+				const FPropertyBagPropertyDesc* DefaultPropertyDesc = DefaultBag->FindPropertyDescByName(Property->GetFName());
+				if (PropertyDesc
+					&& DefaultPropertyDesc
+					&& OverrideProvider.IsPropertyOverridden(PropertyDesc->ID))
+				{
+					UE::StructUtils::Private::CopyPropertyValue(DefaultPropertyDesc, DefaultPropertyBag, PropertyDesc, PropertyBag);
+				}
+			}
+			return true;
+		});
+
+	ChildPropertyHandle->NotifyPostChange(EPropertyChangeType::ValueSet);
+	ChildPropertyHandle->NotifyFinishedChangingProperties();
 }
 
 TSharedRef<SWidget> FPropertyBagInstanceDataDetails::OnPropertyNameContent(TSharedPtr<IPropertyHandle> ChildPropertyHandle, TSharedPtr<SInlineEditableTextBlock> InlineWidget) const
diff --git a/Engine/Plugins/Experimental/StructUtils/Source/StructUtilsEditor/Public/PropertyBagDetails.h b/Engine/Plugins/Experimental/StructUtils/Source/StructUtilsEditor/Public/PropertyBagDetails.h
--- a/Engine/Plugins/Experimental/StructUtils/Source/StructUtilsEditor/Public/PropertyBagDetails.h
+++ b/Engine/Plugins/Experimental/StructUtils/Source/StructUtilsEditor/Public/PropertyBagDetails.h
@@ -57,10 +57,85 @@ private:
 class STRUCTUTILSEDITOR_API FPropertyBagInstanceDataDetails : public FInstancedStructDataDetails
 {
 public:
+
 	FPropertyBagInstanceDataDetails(TSharedPtr<IPropertyHandle> InStructProperty, const TSharedPtr<IPropertyUtilities>& InPropUtils, const bool bInFixedLayout, const bool bInAllowArrays = true);
 
 	virtual void OnChildRowAdded(IDetailPropertyRow& ChildRow) override;
+	
+	/** Enum describing if a property is overridden, or undetermined (e.g. multiselection) */
+	enum class EPropertyOverrideState
+	{
+		Yes,
+		No,
+		Undetermined,
+	};
+
+	/** Interface to allow to modify override status of a specific parameter. */
+	struct IPropertyBagOverrideProvider
+	{
+		virtual ~IPropertyBagOverrideProvider()
+		{
+		}
+
+		virtual bool IsPropertyOverridden(const FGuid PropertyID) const = 0;
+		virtual void SetPropertyOverride(const FGuid PropertyID, const bool bIsOverridden) const = 0;
+	};
+
+	/**
+	 * Callback function for EnumeratePropertyBags.
+	 * @return true to continue enumeration
+	 */
+	using EnumeratePropertyBagFuncRef = TFunctionRef<bool(const FInstancedPropertyBag& /*DefaultPropertyBag*/, FInstancedPropertyBag& /*PropertyBag*/, IPropertyBagOverrideProvider& /*OverrideProvider*/)>;
+
+	/**
+	 * Method that is called to determine if a derived class has property override logic implemented.
+	 * If true is returned, the overridden class is expected to implement PreChangeOverrides(), PostChangeOverrides(), EnumeratePropertyBags().
+	 * @return true of derived class has override logic implemented.
+	 */
+	virtual bool HasPropertyOverrides() const
+	{
+		return false;
+	}
+
+	/** Called before property override is changed. */
+	virtual void PreChangeOverrides()
+	{
+		// ParentPropertyHandle would be the property that holds the property overrides. 
+		//		ParentPropertyHandle->NotifyPreChange();
+		checkf(false, TEXT("PreChangeOverrides() is expecgted to be implemented when HasPropertyOverrides() returns true."));
+	}
 
+	/** Called after property override is changed. */
+	virtual void PostChangeOverrides()
+	{
+		// ParentPropertyHandle is be the property that holds the property overrides. 
+		//		ParentPropertyHandle->NotifyPostChange(EPropertyChangeType::ValueSet);
+		//		ParentPropertyHandle->NotifyFinishedChangingProperties();
+		checkf(false, TEXT("PostChangeOverrides() is expecgted to be implemented when HasPropertyOverrides() returns true."));
+	}
+
+	/**
+	 * Called to enumerate each property bag on the property handle.
+	 * The Func expects DefaultPropertyBag (the values that are override), and PropertyBag (the one that PropertyBagHandle points to),
+	 * and instance of IPropertyBagOverrideProvider which is used to query if specific property is overridden, or to set the property override state.
+	 */
+	virtual void EnumeratePropertyBags(TSharedPtr<IPropertyHandle> PropertyBagHandle, const EnumeratePropertyBagFuncRef& Func) const
+	{
+		checkf(false, TEXT("EnumeratePropertyBags() is expecgted to be implemented when HasPropertyOverrides() returns true."));
+	}
+
+	/** @return true if property of specified child property is overridden. */
+	virtual EPropertyOverrideState IsPropertyOverridden(TSharedPtr<IPropertyHandle> ChildPropertyHandle) const;
+	
+	/** Called to set the override state of specified child property. */
+	virtual void SetPropertyOverride(TSharedPtr<IPropertyHandle> ChildPropertyHandle, const bool bIsOverridden);
+
+	/** @return true if the child property has default value. */
+	virtual bool IsDefaultValue(TSharedPtr<IPropertyHandle> ChildPropertyHandle) const;
+
+	/** Called to reset the child property to default value. */
+	virtual void ResetToDefault(TSharedPtr<IPropertyHandle> ChildPropertyHandle);
+	
 protected:
 	TSharedRef<SWidget> OnPropertyNameContent(TSharedPtr<IPropertyHandle> ChildPropertyHandle, TSharedPtr<SInlineEditableTextBlock> InlineWidget) const;
 
diff --git a/Engine/Plugins/Runtime/StateTree/Source/StateTreeEditorModule/Private/Customizations/StateTreeReferenceDetails.cpp b/Engine/Plugins/Runtime/StateTree/Source/StateTreeEditorModule/Private/Customizations/StateTreeReferenceDetails.cpp
--- a/Engine/Plugins/Runtime/StateTree/Source/StateTreeEditorModule/Private/Customizations/StateTreeReferenceDetails.cpp
+++ b/Engine/Plugins/Runtime/StateTree/Source/StateTreeEditorModule/Private/Customizations/StateTreeReferenceDetails.cpp
@@ -6,12 +6,92 @@
 #include "StateTreeDelegates.h"
 #include "IDetailChildrenBuilder.h"
 #include "IPropertyUtilities.h"
+#include "PropertyBagDetails.h"
 #include "PropertyCustomizationHelpers.h"
 #include "StateTreeReference.h"
+#include "StateTreePropertyHelpers.h"
 #include "ScopedTransaction.h"
 
 #define LOCTEXT_NAMESPACE "StateTreeEditor"
 
+class FStateTreeRefParametersDetails : public FPropertyBagInstanceDataDetails
+{
+public:
+	FStateTreeRefParametersDetails(const TSharedPtr<IPropertyHandle> InStateTreeRefStructProperty, const TSharedPtr<IPropertyHandle> InParametersStructProperty, const TSharedPtr<IPropertyUtilities>& InPropUtils)
+		: FPropertyBagInstanceDataDetails(InParametersStructProperty, InPropUtils, /*bInFixedLayout*/true)
+		, StateTreeRefStructProperty(InStateTreeRefStructProperty)
+	{
+		check(UE::StateTree::PropertyHelpers::IsScriptStruct<FStateTreeReference>(StateTreeRefStructProperty));
+	}
+
+protected:
+	struct FStateTreeReferenceOverrideProvider : public IPropertyBagOverrideProvider
+	{
+		FStateTreeReferenceOverrideProvider(FStateTreeReference& InStateTreeRef)
+			: StateTreeRef(InStateTreeRef)
+		{
+		}
+		
+		virtual bool IsPropertyOverridden(const FGuid PropertyID) const override
+		{
+			return StateTreeRef.IsPropertyOverridden(PropertyID);
+		}
+		
+		virtual void SetPropertyOverride(const FGuid PropertyID, const bool bIsOverridden) const override
+		{
+			StateTreeRef.SetPropertyOverridden(PropertyID, bIsOverridden);
+		}
+
+	private:
+		FStateTreeReference& StateTreeRef;
+	};
+
+	virtual bool HasPropertyOverrides() const override
+	{
+		return true;
+	}
+
+	virtual void PreChangeOverrides() override
+	{
+		check(StateTreeRefStructProperty);
+		StateTreeRefStructProperty->NotifyPreChange();
+	}
+
+	virtual void PostChangeOverrides() override
+	{
+		check(StateTreeRefStructProperty);
+		StateTreeRefStructProperty->NotifyPostChange(EPropertyChangeType::ValueSet);
+		StateTreeRefStructProperty->NotifyFinishedChangingProperties();
+	}
+
+	virtual void EnumeratePropertyBags(TSharedPtr<IPropertyHandle> PropertyBagHandle, const EnumeratePropertyBagFuncRef& Func) const override
+	{
+		check(StateTreeRefStructProperty);
+		StateTreeRefStructProperty->EnumerateRawData([Func](void* RawData, const int32 /*DataIndex*/, const int32 /*NumDatas*/)
+		{
+			if (FStateTreeReference* StateTreeRef = static_cast<FStateTreeReference*>(RawData))
+			{
+				if (const UStateTree* StateTree = StateTreeRef->GetStateTree())
+				{
+					const FInstancedPropertyBag& DefaultParameters = StateTree->GetDefaultParameters();
+					FInstancedPropertyBag& Parameters = StateTreeRef->GetMutableParameters();
+					FStateTreeReferenceOverrideProvider OverrideProvider(*StateTreeRef);
+					if (!Func(DefaultParameters, Parameters, OverrideProvider))
+					{
+						return false;
+					}
+				}
+			}
+			return true;
+		});
+	}
+
+private:
+	
+	TSharedPtr<IPropertyHandle> StateTreeRefStructProperty;
+};
+
+
 TSharedRef<IPropertyTypeCustomization> FStateTreeReferenceDetails::MakeInstance()
 {
 	return MakeShareable(new FStateTreeReferenceDetails);
@@ -22,6 +102,9 @@ void FStateTreeReferenceDetails::CustomizeHeader(const TSharedRef<IPropertyHandl
 	StructPropertyHandle = InStructPropertyHandle;
 	PropUtils = InCustomizationUtils.GetPropertyUtilities();
 
+	ParametersHandle = StructPropertyHandle->GetChildHandle(GET_MEMBER_NAME_CHECKED(FStateTreeReference, Parameters));
+	check(ParametersHandle);
+	
 	// Make sure parameters are synced when displaying StateTreeReference.
 	// Associated StateTree asset might have been recompiled after the StateTreeReference was loaded.
 	// Note: SyncParameters() will create an undoable transaction, do not try to sync when called during Undo/redo as it would overwrite the undo.   
@@ -30,7 +113,7 @@ void FStateTreeReferenceDetails::CustomizeHeader(const TSharedRef<IPropertyHandl
 		SyncParameters();
 	}
 
-	const TSharedPtr<IPropertyHandle> StateTreeProperty = StructPropertyHandle->GetChildHandle(FName(TEXT("StateTree")));
+	const TSharedPtr<IPropertyHandle> StateTreeProperty = StructPropertyHandle->GetChildHandle(GET_MEMBER_NAME_CHECKED(FStateTreeReference, StateTree));
 	check(StateTreeProperty.IsValid());
 
 	const FString SchemaMetaDataValue = InStructPropertyHandle->GetMetaData(UE::StateTree::SchemaTag);
@@ -65,18 +148,9 @@ void FStateTreeReferenceDetails::CustomizeHeader(const TSharedRef<IPropertyHandl
 
 void FStateTreeReferenceDetails::CustomizeChildren(const TSharedRef<IPropertyHandle> InStructPropertyHandle, IDetailChildrenBuilder& InChildrenBuilder, IPropertyTypeCustomizationUtils& InCustomizationUtils)
 {
-	uint32 NumChildren = 0;
-	InStructPropertyHandle->GetNumChildren(NumChildren);
-	for (uint32 ChildIndex = 0; ChildIndex < NumChildren; ChildIndex++)
-	{
-		const TSharedRef<IPropertyHandle> ChildPropertyHandle = InStructPropertyHandle->GetChildHandle(ChildIndex).ToSharedRef();
-
-		// Skip StateTree that was used for the Header row
-		if (ChildPropertyHandle->GetProperty()->GetFName() != FName(TEXT("StateTree")))
-		{
-			InChildrenBuilder.AddProperty(ChildPropertyHandle);
-		}
-	}
+	// Show parameter values.
+	const TSharedRef<FStateTreeRefParametersDetails> ParameterInstanceDetails = MakeShareable(new FStateTreeRefParametersDetails(StructPropertyHandle, ParametersHandle, PropUtils));
+	InChildrenBuilder.AddCustomBuilder(ParameterInstanceDetails);
 }
 
 void FStateTreeReferenceDetails::OnTreeCompiled(const UStateTree& StateTree) const
diff --git a/Engine/Plugins/Runtime/StateTree/Source/StateTreeEditorModule/Private/Customizations/StateTreeReferenceDetails.h b/Engine/Plugins/Runtime/StateTree/Source/StateTreeEditorModule/Private/Customizations/StateTreeReferenceDetails.h
--- a/Engine/Plugins/Runtime/StateTree/Source/StateTreeEditorModule/Private/Customizations/StateTreeReferenceDetails.h
+++ b/Engine/Plugins/Runtime/StateTree/Source/StateTreeEditorModule/Private/Customizations/StateTreeReferenceDetails.h
@@ -35,5 +35,6 @@ private:
 	void SyncParameters(const UStateTree* StateTreeToSync = nullptr) const;
 
 	TSharedPtr<IPropertyHandle> StructPropertyHandle;
+	TSharedPtr<IPropertyHandle> ParametersHandle;
 	TSharedPtr<IPropertyUtilities> PropUtils;
 };
diff --git a/Engine/Plugins/Runtime/StateTree/Source/StateTreeEditorModule/Private/StateTreePropertyHelpers.h b/Engine/Plugins/Runtime/StateTree/Source/StateTreeEditorModule/Private/StateTreePropertyHelpers.h
--- a/Engine/Plugins/Runtime/StateTree/Source/StateTreeEditorModule/Private/StateTreePropertyHelpers.h
+++ b/Engine/Plugins/Runtime/StateTree/Source/StateTreeEditorModule/Private/StateTreePropertyHelpers.h
@@ -22,6 +22,19 @@ void DispatchPostEditToNodes(UObject& Owner, FPropertyChangedChainEvent& Propert
 /** Makes deterministic ID from the owners property path, a property path (or any string), and a seed value (e.g. array index). */
 FGuid MakeDeterministicID(const UObject& Owner, const FString& PropertyPath, const uint64 Seed);
 
+/* @return true if the property handle points to struct property of specified type.*/
+template<typename T>
+bool IsScriptStruct(const TSharedPtr<IPropertyHandle>& PropertyHandle)
+{
+	if (!PropertyHandle)
+	{
+		return false;
+	}
+
+	FStructProperty* StructProperty = CastField<FStructProperty>(PropertyHandle->GetProperty());
+	return StructProperty && StructProperty->Struct->IsA(TBaseStructure<T>::Get()->GetClass());
+}
+
 /**
  * Gets a struct value from property handle, checks type before access. Expects T is struct.
  * @param ValueProperty Handle to property where value is got from.
diff --git a/Engine/Plugins/Runtime/StateTree/Source/StateTreeModule/Private/StateTreeReference.cpp b/Engine/Plugins/Runtime/StateTree/Source/StateTreeModule/Private/StateTreeReference.cpp
--- a/Engine/Plugins/Runtime/StateTree/Source/StateTreeModule/Private/StateTreeReference.cpp
+++ b/Engine/Plugins/Runtime/StateTree/Source/StateTreeModule/Private/StateTreeReference.cpp
@@ -5,6 +5,32 @@
 
 #include UE_INLINE_GENERATED_CPP_BY_NAME(StateTreeReference)
 
+void FStateTreeReference::SyncParameters()
+{
+	if (StateTree == nullptr)
+	{
+		Parameters.Reset();
+	}
+	else
+	{
+		// In editor builds, sync with overrides.
+		Parameters.MigrateToNewBagInstanceWithOverrides(StateTree->GetDefaultParameters(), PropertyOverrides);
+		
+		// Remove overrides that do not exists anymore
+		if (!PropertyOverrides.IsEmpty())
+		{
+			const UPropertyBag* Bag = Parameters.GetPropertyBagStruct();
+			for (TArray<FGuid>::TIterator It = PropertyOverrides.CreateIterator(); It; ++It)
+			{
+				if (!Bag->FindPropertyDescByID(*It))
+				{
+					It.RemoveCurrentSwap();
+				}
+			}
+		}
+	}
+}
+
 void FStateTreeReference::SyncParametersToMatchStateTree(FInstancedPropertyBag& ParametersToSync) const
 {
 	if (StateTree == nullptr)
@@ -19,8 +45,53 @@ void FStateTreeReference::SyncParametersToMatchStateTree(FInstancedPropertyBag&
 
 bool FStateTreeReference::RequiresParametersSync() const
 {
-	return (StateTree == nullptr && Parameters.IsValid())
-		|| (StateTree != nullptr && StateTree->GetDefaultParameters().GetPropertyBagStruct() != Parameters.GetPropertyBagStruct());
+	bool bShouldSync = false;
+	
+	if (StateTree)
+	{
+		const FInstancedPropertyBag& DefaultParameters = StateTree->GetDefaultParameters();
+		const UPropertyBag* DefaultParametersBag = DefaultParameters.GetPropertyBagStruct();
+		const UPropertyBag* ParametersBag = Parameters.GetPropertyBagStruct();
+		
+		// Mismatching property bags, needs sync.
+		if (DefaultParametersBag != ParametersBag)
+		{
+			bShouldSync = true;
+		}
+		else if (ParametersBag && DefaultParametersBag)
+		{
+			// Check if non-overridden parameters are not identical, needs sync.
+			const uint8* SourceAddress = DefaultParameters.GetValue().GetMemory();
+			const uint8* TargetAddress = Parameters.GetValue().GetMemory();
+			check(SourceAddress);
+			check(TargetAddress);
+
+			for (const FPropertyBagPropertyDesc& Desc : ParametersBag->GetPropertyDescs())
+			{
+				// Skip overridden
+				if (PropertyOverrides.Contains(Desc.ID))
+				{
+					continue;
+				}
+
+				const uint8* SourceValueAddress = SourceAddress + Desc.CachedProperty->GetOffset_ForInternal();
+				const uint8* TargetValueAddress = TargetAddress + Desc.CachedProperty->GetOffset_ForInternal();
+				if (!Desc.CachedProperty->Identical(SourceValueAddress, TargetValueAddress))
+				{
+					// Mismatching values, should sync.
+					bShouldSync = true;
+					break;
+				}
+			}
+		}
+	}
+	else
+	{
+		// Empty state tree should not have parameters
+		bShouldSync = Parameters.IsValid();
+	}
+	
+	return bShouldSync;
 }
 
 void FStateTreeReference::ConditionallySyncParameters() const
@@ -33,3 +104,37 @@ void FStateTreeReference::ConditionallySyncParameters() const
 	}
 }
 
+void FStateTreeReference::SetPropertyOverridden(const FGuid PropertyID, const bool bIsOverridden)
+{
+	if (bIsOverridden)
+	{
+		PropertyOverrides.AddUnique(PropertyID);
+	}
+	else
+	{
+		PropertyOverrides.Remove(PropertyID);
+		ConditionallySyncParameters();
+	}
+}
+
+bool FStateTreeReference::Serialize(FStructuredArchive::FSlot Slot)
+{
+	Slot.GetUnderlyingArchive().UsingCustomVersion(FStateTreeCustomVersion::GUID);
+	return false; // Let the default serializer handle serializing.
+}
+
+void FStateTreeReference::PostSerialize(const FArchive& Ar)
+{
+	const int32 CurrentVersion = Ar.CustomVer(FStateTreeCustomVersion::GUID);
+	if (CurrentVersion < FStateTreeCustomVersion::OverridableParameters)
+	{
+		// In earlier versions, all parameters were overwritten.
+		if (const UPropertyBag* Bag = Parameters.GetPropertyBagStruct())
+		{
+			for (const FPropertyBagPropertyDesc& Desc : Bag->GetPropertyDescs())
+			{
+				PropertyOverrides.Add(Desc.ID);
+			}
+		}
+	}
+}
diff --git a/Engine/Plugins/Runtime/StateTree/Source/StateTreeModule/Public/StateTree.h b/Engine/Plugins/Runtime/StateTree/Source/StateTreeModule/Public/StateTree.h
--- a/Engine/Plugins/Runtime/StateTree/Source/StateTreeModule/Public/StateTree.h
+++ b/Engine/Plugins/Runtime/StateTree/Source/StateTreeModule/Public/StateTree.h
@@ -43,6 +43,8 @@ struct STATETREEMODULE_API FStateTreeCustomVersion
 		AddedLinkedAssetState,
 		// Change how external data is accessed
 		ChangedExternalDataAccess,
+		// Added override option for parameters
+		OverridableParameters,
 
 		// -----<new versions can be added above this line>-------------------------------------------------
 		VersionPlusOne,
diff --git a/Engine/Plugins/Runtime/StateTree/Source/StateTreeModule/Public/StateTreeReference.h b/Engine/Plugins/Runtime/StateTree/Source/StateTreeModule/Public/StateTreeReference.h
--- a/Engine/Plugins/Runtime/StateTree/Source/StateTreeModule/Public/StateTreeReference.h
+++ b/Engine/Plugins/Runtime/StateTree/Source/StateTreeModule/Public/StateTreeReference.h
@@ -57,11 +57,12 @@ struct STATETREEMODULE_API FStateTreeReference
 	/**
 	 * Enforce self parameters to be compatible with those exposed by the selected StateTree asset.
 	 */
-	void SyncParameters() { SyncParametersToMatchStateTree(Parameters); }
+	void SyncParameters();
 
 	/**
 	 * Sync provided parameters to be compatible with those exposed by the selected StateTree asset.
 	 */
+	UE_DEPRECATED(5.4, "Use SyncParameters() instead.")
 	void SyncParametersToMatchStateTree(FInstancedPropertyBag& ParametersToSync) const;
 
 	/**
@@ -72,15 +73,44 @@ struct STATETREEMODULE_API FStateTreeReference
 
 	/** Sync parameters to match the asset if required. */
 	void ConditionallySyncParameters() const;
-	
+
+	/** @return true if the property of specified ID is overridden. */
+	bool IsPropertyOverridden(const FGuid PropertyID) const
+	{
+		return PropertyOverrides.Contains(PropertyID);
+	}
+
+	/** Sets the override status of specified property by ID. */
+	void SetPropertyOverridden(const FGuid PropertyID, const bool bIsOverridden);
+
+	bool Serialize(FStructuredArchive::FSlot Slot);
+	void PostSerialize(const FArchive& Ar);
+
 protected:
 	UPROPERTY(EditAnywhere, Category = "")
 	TObjectPtr<UStateTree> StateTree = nullptr;
 
 	UPROPERTY(EditAnywhere, Category = "", meta = (FixedLayout))
 	FInstancedPropertyBag Parameters;
+
+	/** Array of overridden properties. Non-overridden properties will inherit the values from the StateTree default parameters. */
+	UPROPERTY(EditAnywhere, Category = "")
+	TArray<FGuid> PropertyOverrides;
+
+	friend class FStateTreeReferenceDetails;
 };
 
+template<>
+struct TStructOpsTypeTraits<FStateTreeReference> : public TStructOpsTypeTraitsBase2<FStateTreeReference>
+{
+	enum
+	{
+		WithStructuredSerializer = true,
+		WithPostSerialize = true,
+	};
+};
+
+
 #if UE_ENABLE_INCLUDE_ORDER_DEPRECATED_IN_5_2
 #include "StateTreeTypes.h"
 #endif
