StructUtils: Small fixes to property bag
- Clarified property migrate method comments
- Added FindPropertyDescByProperty() to make property lookup faster and to encourage right usafe over FindPropertyDescByName() (which uses sanitized name)

diff --git a/Engine/Plugins/Experimental/StructUtils/Source/StructUtils/Private/PropertyBag.cpp b/Engine/Plugins/Experimental/StructUtils/Source/StructUtils/Private/PropertyBag.cpp
--- a/Engine/Plugins/Experimental/StructUtils/Source/StructUtils/Private/PropertyBag.cpp
+++ b/Engine/Plugins/Experimental/StructUtils/Source/StructUtils/Private/PropertyBag.cpp
@@ -2902,3 +2902,8 @@ const FPropertyBagPropertyDesc* UPropertyBag::FindPropertyDescByPropertyName(con
 {
 	return PropertyDescs.FindByPredicate([&PropertyName](const FPropertyBagPropertyDesc& Desc) { return Desc.CachedProperty && Desc.CachedProperty->GetFName() == PropertyName; });
 }
+
+const FPropertyBagPropertyDesc* UPropertyBag::FindPropertyDescByProperty(const FProperty* Property) const
+{
+	return PropertyDescs.FindByPredicate([&Property](const FPropertyBagPropertyDesc& Desc) { return Desc.CachedProperty && Desc.CachedProperty == Property; });
+}
diff --git a/Engine/Plugins/Experimental/StructUtils/Source/StructUtils/Public/PropertyBag.h b/Engine/Plugins/Experimental/StructUtils/Source/StructUtils/Public/PropertyBag.h
--- a/Engine/Plugins/Experimental/StructUtils/Source/StructUtils/Public/PropertyBag.h
+++ b/Engine/Plugins/Experimental/StructUtils/Source/StructUtils/Public/PropertyBag.h
@@ -434,20 +434,21 @@ struct STRUCTUTILS_API FInstancedPropertyBag
 	void MigrateToNewBagStruct(const UPropertyBag* NewBagStruct);
 
 	/**
-	 * Changes the type of this bag to the specified other bag, copies base values from the other bag, and migrates existing values over.
+	 * Changes the type of this bag to the InNewBagInstance, and migrates existing values over.
+	 * Properties that do not exist in this bag will get values from NewBagInstance.
 	 * The properties are matched between the bags based on the property ID.
-	 * @param NewBagInstance Reference to the new type.
+	 * @param InNewBagInstance New bag composition and values used for new properties.
 	 */
-	void MigrateToNewBagInstance(const FInstancedPropertyBag& NewBagInstance);
+	void MigrateToNewBagInstance(const FInstancedPropertyBag& InNewBagInstance);
 
 	/**
-	 * Changes the type of this bag to the specified other bag, copies base values from the other bag,
-	 * and migrates existing values over if marked as overridden in the OverriddenPropertyIDs.
+	 * Changes the type of this bag to the InNewBagInstance, and migrates existing values over if marked as overridden in the OverriddenPropertyIDs.
+	 * Properties that does not exist in this bag, or are not overridden, will get values from InNewBagInstance.
 	 * The properties are matched between the bags based on the property ID.
-	 * @param NewBagInstance Reference to the new type.
+	 * @param InNewBagInstance New bag composition and values used for new properties.
 	 * @param OverriddenPropertyIDs Array if property IDs which should be copied over to the new instance. 
 	 */
-	void MigrateToNewBagInstanceWithOverrides(const FInstancedPropertyBag& NewBagInstance, TConstArrayView<FGuid> OverriddenPropertyIDs);
+	void MigrateToNewBagInstanceWithOverrides(const FInstancedPropertyBag& InNewBagInstance, TConstArrayView<FGuid> OverriddenPropertyIDs);
 
 	/** @return pointer to the property bag struct. */ 
 	const UPropertyBag* GetPropertyBagStruct() const;
@@ -886,6 +887,9 @@ public:
 	/** @return property description based on the created property name. The name can be different from the descriptor name due to name sanitization. */
 	const FPropertyBagPropertyDesc* FindPropertyDescByPropertyName(const FName PropertyName) const;
 
+	/** @return property description based on pointer to property. */
+	const FPropertyBagPropertyDesc* FindPropertyDescByProperty(const FProperty* Property) const;
+
 #if WITH_ENGINE && WITH_EDITOR
 	/** @return true if any of the properties on the bag has type of the specified user defined struct. */
 	bool ContainsUserDefinedStruct(const UUserDefinedStruct* UserDefinedStruct) const;
diff --git a/Engine/Plugins/Experimental/StructUtils/Source/StructUtilsEditor/Private/Customizations/PropertyBagDetails.cpp b/Engine/Plugins/Experimental/StructUtils/Source/StructUtilsEditor/Private/Customizations/PropertyBagDetails.cpp
--- a/Engine/Plugins/Experimental/StructUtils/Source/StructUtilsEditor/Private/Customizations/PropertyBagDetails.cpp
+++ b/Engine/Plugins/Experimental/StructUtils/Source/StructUtilsEditor/Private/Customizations/PropertyBagDetails.cpp
@@ -997,7 +997,7 @@ FPropertyBagInstanceDataDetails::EPropertyOverrideState FPropertyBagInstanceData
 			NumValues++;
 			if (const UPropertyBag* Bag = PropertyBag.GetPropertyBagStruct())
 			{
-				const FPropertyBagPropertyDesc* PropertyDesc = Bag->FindPropertyDescByPropertyName(Property->GetFName());
+				const FPropertyBagPropertyDesc* PropertyDesc = Bag->FindPropertyDescByProperty(Property);
 				if (PropertyDesc && OverrideProvider.IsPropertyOverridden(PropertyDesc->ID))
 				{
 					NumOverrides++;
@@ -1039,7 +1039,7 @@ void FPropertyBagInstanceDataDetails::SetPropertyOverride(TSharedPtr<IPropertyHa
 		{
 			if (const UPropertyBag* Bag = PropertyBag.GetPropertyBagStruct())
 			{
-				if (const FPropertyBagPropertyDesc* PropertyDesc = Bag->FindPropertyDescByPropertyName(Property->GetFName()))
+				if (const FPropertyBagPropertyDesc* PropertyDesc = Bag->FindPropertyDescByProperty(Property))
 				{
 					OverrideProvider.SetPropertyOverride(PropertyDesc->ID, bIsOverridden);
 				}
@@ -1076,8 +1076,8 @@ bool FPropertyBagInstanceDataDetails::IsDefaultValue(TSharedPtr<IPropertyHandle>
 			const UPropertyBag* Bag = PropertyBag.GetPropertyBagStruct();
 			if (Bag && DefaultBag)
 			{
-				const FPropertyBagPropertyDesc* PropertyDesc = Bag->FindPropertyDescByName(Property->GetFName());
-				const FPropertyBagPropertyDesc* DefaultPropertyDesc = DefaultBag->FindPropertyDescByName(Property->GetFName());
+				const FPropertyBagPropertyDesc* PropertyDesc = Bag->FindPropertyDescByProperty(Property);
+				const FPropertyBagPropertyDesc* DefaultPropertyDesc = DefaultBag->FindPropertyDescByProperty(Property);
 				if (PropertyDesc
 					&& DefaultPropertyDesc
 					&& OverrideProvider.IsPropertyOverridden(PropertyDesc->ID))
@@ -1122,8 +1122,8 @@ void FPropertyBagInstanceDataDetails::ResetToDefault(TSharedPtr<IPropertyHandle>
 			const UPropertyBag* Bag = PropertyBag.GetPropertyBagStruct();
 			if (Bag && DefaultBag)
 			{
-				const FPropertyBagPropertyDesc* PropertyDesc = Bag->FindPropertyDescByName(Property->GetFName());
-				const FPropertyBagPropertyDesc* DefaultPropertyDesc = DefaultBag->FindPropertyDescByName(Property->GetFName());
+				const FPropertyBagPropertyDesc* PropertyDesc = Bag->FindPropertyDescByProperty(Property);
+				const FPropertyBagPropertyDesc* DefaultPropertyDesc = DefaultBag->FindPropertyDescByProperty(Property);
 				if (PropertyDesc
 					&& DefaultPropertyDesc
 					&& OverrideProvider.IsPropertyOverridden(PropertyDesc->ID))
