Changed FSharedStruct::Identical to return true also when the contained memory pointer wrapped by the type's instances are identical

diff --git a/Engine/Plugins/Experimental/StructUtils/Source/StructUtils/Private/SharedStruct.cpp b/Engine/Plugins/Experimental/StructUtils/Source/StructUtils/Private/SharedStruct.cpp
--- a/Engine/Plugins/Experimental/StructUtils/Source/StructUtils/Private/SharedStruct.cpp
+++ b/Engine/Plugins/Experimental/StructUtils/Source/StructUtils/Private/SharedStruct.cpp
@@ -8,8 +8,10 @@
 
 bool FConstSharedStruct::Identical(const FConstSharedStruct* Other, uint32 PortFlags) const
 {
-	// Only empty is considered equal
-	return Other != nullptr && GetMemory() == nullptr && Other->GetMemory() == nullptr && GetScriptStruct() == nullptr && Other->GetScriptStruct() == nullptr;
+	// Only empty or strictly identical is considered equal
+	return Other != nullptr
+		&& GetMemory() == Other->GetMemory()
+		&& GetScriptStruct() == Other->GetScriptStruct();
 }
 
 void FConstSharedStruct::AddStructReferencedObjects(class FReferenceCollector& Collector)
@@ -38,8 +40,10 @@ FSharedStruct& FSharedStruct::operator=(const FConstStructView& InOther)
 
 bool FSharedStruct::Identical(const FSharedStruct* Other, uint32 PortFlags) const
 {
-	// Only empty is considered equal
-	return Other != nullptr && GetMemory() == nullptr && Other->GetMemory() == nullptr && GetScriptStruct() == nullptr && Other->GetScriptStruct() == nullptr;
+	// Only empty or strictly identical is considered equal
+	return Other != nullptr 
+		&& GetMemory() == Other->GetMemory() 
+		&& GetScriptStruct() == Other->GetScriptStruct();
 }
 
 void FSharedStruct::AddStructReferencedObjects(class FReferenceCollector& Collector)
