Fix PropertyBag not finding existing instances with the same name when they are AsyncLoading
#rnx

diff --git a/Engine/Plugins/Experimental/StructUtils/Source/StructUtils/Private/PropertyBag.cpp b/Engine/Plugins/Experimental/StructUtils/Source/StructUtils/Private/PropertyBag.cpp
--- a/Engine/Plugins/Experimental/StructUtils/Source/StructUtils/Private/PropertyBag.cpp
+++ b/Engine/Plugins/Experimental/StructUtils/Source/StructUtils/Private/PropertyBag.cpp
@@ -2712,9 +2712,15 @@ const UPropertyBag* UPropertyBag::GetOrCreateFromDescs(const TConstArrayView<FPr
 	// partially-constructed objects 
 	FScopeLock ScopeLock(&UE::StructUtils::Private::GPropertyBagLock);
 
-	if (const UPropertyBag* ExistingBag = FindObject<UPropertyBag>(GetTransientPackage(), *ScriptStructName))
+	// we need to use StaticFindObjectFastInternal with ExclusiveInternalFlags = EInternalObjectFlags::None
+	// here because objects with RF_NeedPostLoad cannot be found by regular FindObject calls as they will have
+	// ExclusiveInternalFlags = EInternalObjectFlags::AsyncLoading when called from game thread.
+	if (UObject* ExistingObject = StaticFindObjectFastInternal(UPropertyBag::StaticClass(), GetTransientPackage(), *ScriptStructName, true, RF_NoFlags, EInternalObjectFlags::None))
 	{
-		return ExistingBag;
+		if (const UPropertyBag* ExistingBag = Cast<UPropertyBag>(ExistingObject))
+		{
+			return ExistingBag;
+		}
 	}
 
 	UPropertyBag* NewBag = NewObject<UPropertyBag>(GetTransientPackage(), *ScriptStructName, RF_Standalone | RF_Transient);
