Fix InstancedStruct implicit conversions to base types

TInstancedStruct<FBase> = TInstancedStruct<FDerived> now works.

diff --git a/Engine/Plugins/Experimental/StructUtils/Source/StructUtils/Public/InstancedStruct.h b/Engine/Plugins/Experimental/StructUtils/Source/StructUtils/Public/InstancedStruct.h
--- a/Engine/Plugins/Experimental/StructUtils/Source/StructUtils/Public/InstancedStruct.h
+++ b/Engine/Plugins/Experimental/StructUtils/Source/StructUtils/Public/InstancedStruct.h
@@ -452,6 +452,8 @@ private:
 	 *   This means that any reflected APIs (like ExportText) that accept an FInstancedStruct pointer can also accept a TInstancedStruct pointer.
 	 */
 	FInstancedStruct InstancedStruct;
+
+	template <typename U> friend struct TInstancedStruct;
 };
 
 #if WITH_EDITORONLY_DATA
