FInstancedStruct will now serialize itself when the archive has bIsObjectReferenceCollector.
This is needed for archives that serialize objects with instanced structs that gathers references.

diff --git a/Engine/Plugins/Experimental/StructUtils/Source/StructUtils/Private/InstancedStruct.cpp b/Engine/Plugins/Experimental/StructUtils/Source/StructUtils/Private/InstancedStruct.cpp
--- a/Engine/Plugins/Experimental/StructUtils/Source/StructUtils/Private/InstancedStruct.cpp
+++ b/Engine/Plugins/Experimental/StructUtils/Source/StructUtils/Private/InstancedStruct.cpp
@@ -249,7 +249,7 @@ bool FInstancedStruct::Serialize(FArchive& Ar)
 		Ar << SerialSize;
 		Ar.Seek(FinalOffset);	// Reset archive to its position
 	}
-	else if (Ar.IsCountingMemory() || Ar.IsModifyingWeakAndStrongReferences())
+	else if (Ar.IsCountingMemory() || Ar.IsModifyingWeakAndStrongReferences() || Ar.IsObjectReferenceCollector())
 	{
 		// Report type
 		Ar << NonConstStruct;
