Added 64 bit hashing functions to UE::StructUtils

diff --git a/Engine/Plugins/Experimental/StructUtils/Source/StructUtils/Private/StructUtilsTypes.cpp b/Engine/Plugins/Experimental/StructUtils/Source/StructUtils/Private/StructUtilsTypes.cpp
--- a/Engine/Plugins/Experimental/StructUtils/Source/StructUtils/Private/StructUtilsTypes.cpp
+++ b/Engine/Plugins/Experimental/StructUtils/Source/StructUtils/Private/StructUtilsTypes.cpp
@@ -2,13 +2,16 @@
 
 #include "StructUtilsTypes.h"
 #include "Serialization/ArchiveCrc32.h"
+#include "Hash/CityHash.h"
 #include "StructView.h"
 #include "InstancedStruct.h"
 #include "SharedStruct.h"
+#include "Serialization/ArchiveUObject.h"
+
 
 namespace UE::StructUtils
 {
-	STRUCTUTILS_API uint32 GetStructCrc32(const UScriptStruct& ScriptStruct, const uint8* StructMemory, const uint32 CRC /*= 0*/)
+	uint32 GetStructCrc32(const UScriptStruct& ScriptStruct, const uint8* StructMemory, const uint32 CRC /*= 0*/)
 	{
 		FArchiveCrc32 Ar(HashCombine(CRC, PointerHash(&ScriptStruct)));
 		if (StructMemory)
@@ -29,26 +32,129 @@ namespace UE::StructUtils
 		return 0;
 	}
 
-	STRUCTUTILS_API uint32 GetStructCrc32(const FStructView& StructView, const uint32 CRC /*= 0*/)
+	uint32 GetStructCrc32(const FStructView& StructView, const uint32 CRC /*= 0*/)
 	{
 		return GetStructCrc32Helper(StructView, CRC);
 	}
 
-	STRUCTUTILS_API uint32 GetStructCrc32(const FConstStructView& StructView, const uint32 CRC /*= 0*/)
+	uint32 GetStructCrc32(const FConstStructView& StructView, const uint32 CRC /*= 0*/)
 	{
 		return GetStructCrc32Helper(StructView, CRC);
 	}
 
-	STRUCTUTILS_API uint32 GetStructCrc32(const FSharedStruct& SharedView, const uint32 CRC /*= 0*/)
+	uint32 GetStructCrc32(const FSharedStruct& SharedView, const uint32 CRC /*= 0*/)
 	{
 		return GetStructCrc32Helper(SharedView, CRC);
 	}
 
-	STRUCTUTILS_API uint32 GetStructCrc32(const FConstSharedStruct& SharedView, const uint32 CRC /*= 0*/)
+	uint32 GetStructCrc32(const FConstSharedStruct& SharedView, const uint32 CRC /*= 0*/)
 	{
 		return GetStructCrc32Helper(SharedView, CRC);
 	}
 
+	class FArchiveCityHash64 : public FArchiveUObject
+	{
+	public:
+		explicit FArchiveCityHash64(const uint64 InHash = 0)
+			: Hash(InHash)
+		{}
+
+		/**
+		 * @return The hash computed so far.
+		 */
+		uint64 GetHash() const { return Hash; }
+
+		//~ Begin FArchive Interface
+		virtual void Serialize(void* Data, int64 Num) override
+		{
+			const char* BytePointer = static_cast<const char*>(Data);
+			while (Num > 0)
+			{
+				const int32 BytesToHash = static_cast<int32>(FMath::Min(Num, int64(MAX_int32)));
+
+				Hash = CityHash64WithSeed(BytePointer, BytesToHash, Hash);
+
+				Num -= BytesToHash;
+				BytePointer += BytesToHash;
+			}
+		}
+
+		virtual FArchive& operator<<(class FName& Name) override
+		{
+			Hash = HashString(Name.ToString(), Hash);
+			return *this;
+		}
+
+		virtual FArchive& operator<<(class UObject*& Object) override
+		{
+			Hash = HashString(GetPathNameSafe(Object), Hash);
+			return *this;
+		}
+
+		virtual FString GetArchiveName() const override { return TEXT("FArchiveCityHash64"); }
+		//~ End FArchive Interface
+
+		FORCEINLINE static uint64 HashString(const FString& InString, const uint32 InHash = 0)
+		{
+			return CityHash64WithSeed(reinterpret_cast<const char*>(*InString), InString.GetAllocatedSize(), InHash);
+		}
+
+	private:
+		uint64 Hash;
+	};
+
+	uint64 GetStructHash64(const UScriptStruct& ScriptStruct, const uint8* StructMemory)
+	{
+		const uint64 BaseHash = FArchiveCityHash64::HashString(ScriptStruct.GetPathName());
+
+		const bool bHasTypeHash = ScriptStruct.GetCppStructOps()->HasGetTypeHash();
+		if (bHasTypeHash)
+		{
+			const uint32 StructHash = ScriptStruct.GetStructTypeHash(StructMemory);
+			return CityHash128to64(Uint128_64(BaseHash, StructHash));
+		}
+		else if (StructMemory)
+		{
+			FArchiveCityHash64 Ar(BaseHash);
+			UScriptStruct& NonConstScriptStruct = const_cast<UScriptStruct&>(ScriptStruct);
+			NonConstScriptStruct.SerializeItem(Ar, const_cast<uint8*>(StructMemory), nullptr);
+
+			return Ar.GetHash();
+		}
+		
+		return BaseHash;
+	}
+
+	template<typename T>
+	uint64 GetStructHash64Helper(const T& Struct)
+	{
+		if (const UScriptStruct* ScriptStruct = Struct.GetScriptStruct())
+		{
+			return GetStructHash64(*ScriptStruct, Struct.GetMemory());
+		}
+		return 0;
+	}
+
+	uint64 GetStructHash64(const FStructView& StructView)
+	{
+		return GetStructHash64Helper(StructView);
+	}
+
+	uint64 GetStructHash64(const FConstStructView& StructView)
+	{
+		return GetStructHash64Helper(StructView);
+	}
+
+	uint64 GetStructHash64(const FSharedStruct& SharedView)
+	{
+		return GetStructHash64Helper(SharedView);
+	}
+
+	uint64 GetStructHash64(const FConstSharedStruct& SharedView)
+	{
+		return GetStructHash64Helper(SharedView);
+	}
+
 	namespace Private
 	{
 #if WITH_ENGINE && WITH_EDITOR
diff --git a/Engine/Plugins/Experimental/StructUtils/Source/StructUtils/Public/StructUtilsTypes.h b/Engine/Plugins/Experimental/StructUtils/Source/StructUtils/Public/StructUtilsTypes.h
--- a/Engine/Plugins/Experimental/StructUtils/Source/StructUtils/Public/StructUtilsTypes.h
+++ b/Engine/Plugins/Experimental/StructUtils/Source/StructUtils/Public/StructUtilsTypes.h
@@ -25,6 +25,17 @@ namespace UE::StructUtils
 	extern STRUCTUTILS_API uint32 GetStructCrc32(const FSharedStruct& SharedView, const uint32 CRC = 0);
 	extern STRUCTUTILS_API uint32 GetStructCrc32(const FConstSharedStruct& SharedView, const uint32 CRC = 0);
 
+	/** 
+	 * CityHash64-based struct hashing functions.
+	 * Note that these are relatively slow due to using either UScriptStruct.GetStructTypeHash (if implemented) or
+	 * a serialization path. 
+	 */
+	extern STRUCTUTILS_API uint64 GetStructHash64(const UScriptStruct& ScriptStruct, const uint8* StructMemory);
+	extern STRUCTUTILS_API uint64 GetStructHash64(const FStructView& StructView);
+	extern STRUCTUTILS_API uint64 GetStructHash64(const FConstStructView& StructView);
+	extern STRUCTUTILS_API uint64 GetStructHash64(const FSharedStruct& SharedView);
+	extern STRUCTUTILS_API uint64 GetStructHash64(const FConstSharedStruct& SharedView);
+
 	template <typename T>
 	typename TEnableIf<!TIsDerivedFrom<T, UObject>::IsDerived, UScriptStruct*>::Type GetAsUStruct()
 	{
