Adding support for instance struct picker to have asset reference filter

diff --git a/Engine/Plugins/Experimental/StructUtils/Source/StructUtilsEditor/Private/Customizations/InstancedStructDetails.cpp b/Engine/Plugins/Experimental/StructUtils/Source/StructUtilsEditor/Private/Customizations/InstancedStructDetails.cpp
--- a/Engine/Plugins/Experimental/StructUtils/Source/StructUtilsEditor/Private/Customizations/InstancedStructDetails.cpp
+++ b/Engine/Plugins/Experimental/StructUtils/Source/StructUtilsEditor/Private/Customizations/InstancedStructDetails.cpp
@@ -3,6 +3,9 @@
 #include "InstancedStructDetails.h"
 #include "DetailWidgetRow.h"
 #include "DetailLayoutBuilder.h"
+#include "Editor.h"
+#include "Editor/AssetReferenceFilter.h"
+#include "Editor/EditorEngine.h"
 #include "IDetailChildrenBuilder.h"
 #include "IPropertyUtilities.h"
 #include "UObject/Package.h"
@@ -177,6 +180,14 @@ bool FInstancedStructFilter::IsStructAllowed(const FStructViewerInitializationOp
 		return false;
 	}
 
+	if (AssetReferenceFilter.IsValid())
+	{
+		if (!AssetReferenceFilter->PassesFilter(FAssetData(InStruct)))
+		{
+			return false;
+		}
+	}
+
 	// Query the native struct to see if it has the correct parent type (if any)
 	return !BaseStruct || InStruct->IsChildOf(BaseStruct);
 }
@@ -495,6 +506,20 @@ TSharedRef<SWidget> FInstancedStructDetails::GenerateStructPicker()
 	StructFilter->bAllowUserDefinedStructs = BaseScriptStruct == nullptr; // Only allow user defined structs when BaseStruct is not set.
 	StructFilter->bAllowBaseStruct = !bExcludeBaseStruct;
 
+	if (GEditor && StructProperty)
+	{
+		FAssetReferenceFilterContext AssetReferenceFilterContext;
+		
+		TArray<UPackage*> OuterPackages;
+		StructProperty->GetOuterPackages(OuterPackages);
+		for (UPackage* OuterPackage : OuterPackages)
+		{
+			AssetReferenceFilterContext.ReferencingAssets.Add(FAssetData(OuterPackage));
+		}
+
+		StructFilter->AssetReferenceFilter = GEditor->MakeAssetReferenceFilter(AssetReferenceFilterContext);
+	}
+
 	const UScriptStruct* SelectedStruct = nullptr;
 	const FPropertyAccess::Result Result = UE::StructUtils::Private::GetCommonScriptStruct(StructProperty, SelectedStruct);
 	
diff --git a/Engine/Plugins/Experimental/StructUtils/Source/StructUtilsEditor/Public/InstancedStructDetails.h b/Engine/Plugins/Experimental/StructUtils/Source/StructUtilsEditor/Public/InstancedStructDetails.h
--- a/Engine/Plugins/Experimental/StructUtils/Source/StructUtilsEditor/Public/InstancedStructDetails.h
+++ b/Engine/Plugins/Experimental/StructUtils/Source/StructUtilsEditor/Public/InstancedStructDetails.h
@@ -7,6 +7,7 @@
 #include "EditorUndoClient.h"
 #include "StructViewerFilter.h"
 
+class IAssetReferenceFilter;
 class IPropertyHandle;
 class IDetailPropertyRow;
 class IPropertyHandle;
@@ -128,6 +129,9 @@ public:
 
 	virtual bool IsStructAllowed(const FStructViewerInitializationOptions& InInitOptions, const UScriptStruct* InStruct, TSharedRef<FStructViewerFilterFuncs> InFilterFuncs) override;
 	virtual bool IsUnloadedStructAllowed(const FStructViewerInitializationOptions& InInitOptions, const FSoftObjectPath& InStructPath, TSharedRef<FStructViewerFilterFuncs> InFilterFuncs) override;
+
+ 	// Optional filter to prevent selection of some structs e.g. ones in a plugin that is inaccessible from the object being edited
+ 	TSharedPtr<IAssetReferenceFilter> AssetReferenceFilter;
 };
 
 
