PR #11667: FIX: FInstancedStructDetails did not respect StructPropertyHandle->IsEditable()

diff --git a/Engine/Plugins/Experimental/StructUtils/Source/StructUtilsEditor/Private/Customizations/InstancedStructDetails.cpp b/Engine/Plugins/Experimental/StructUtils/Source/StructUtilsEditor/Private/Customizations/InstancedStructDetails.cpp
--- a/Engine/Plugins/Experimental/StructUtils/Source/StructUtilsEditor/Private/Customizations/InstancedStructDetails.cpp
+++ b/Engine/Plugins/Experimental/StructUtils/Source/StructUtilsEditor/Private/Customizations/InstancedStructDetails.cpp
@@ -416,7 +416,8 @@ void FInstancedStructDetails::CustomizeHeader(TSharedRef<class IPropertyHandle>
 					.Font(IDetailLayoutBuilder::GetDetailFont())
 				]
 			]
-		];
+		]
+		.IsEnabled(StructProperty->IsEditable());
 }
 
 void FInstancedStructDetails::OnObjectsReinstanced(const FReplacementObjectMap& ObjectMap)
