Added functions to FSharedStruct and FConstSharedStruct to compare instanced by value

diff --git a/Engine/Plugins/Experimental/StructUtils/Source/StructUtils/Public/SharedStruct.h b/Engine/Plugins/Experimental/StructUtils/Source/StructUtils/Public/SharedStruct.h
--- a/Engine/Plugins/Experimental/StructUtils/Source/StructUtils/Public/SharedStruct.h
+++ b/Engine/Plugins/Experimental/StructUtils/Source/StructUtils/Public/SharedStruct.h
@@ -322,6 +322,19 @@ struct STRUCTUTILS_API FSharedStruct
 		return !operator==(Other);
 	}
 
+	/**
+	 * Determines whether Other contains same values as `this`
+	 * @return whether the values stored are equal
+	 */
+	template <typename OtherType>
+	bool CompareStructValues(const OtherType& Other, uint32 PortFlags = 0) const
+	{
+		UE::StructUtils::CheckWrapperType<OtherType>();
+		const UScriptStruct* ScriptStruct = GetScriptStruct();
+		return (ScriptStruct == Other.GetScriptStruct())
+			&& (!ScriptStruct || ScriptStruct->CompareScriptStruct(GetMemory(), Other.GetMemory(), PortFlags));
+	}
+
 protected:
 	TSharedPtr<FStructSharedMemory> StructMemoryPtr;
 };
@@ -533,6 +546,19 @@ struct STRUCTUTILS_API FConstSharedStruct
 		return !operator==(Other);
 	}
 
+	/** 
+	 * Determines whether Other contains same values as `this` 
+	 * @return whether the values stored are equal
+	 */
+	template <typename OtherType>
+	bool CompareStructValues(const OtherType& Other, uint32 PortFlags = 0) const
+	{
+		UE::StructUtils::CheckWrapperType<OtherType>();
+		const UScriptStruct* ScriptStruct = GetScriptStruct();
+		return (ScriptStruct == Other.GetScriptStruct())
+			&& (!ScriptStruct || ScriptStruct->CompareScriptStruct(GetMemory(), Other.GetMemory(), PortFlags));
+	}
+
 protected:
 
 	TSharedPtr<const FStructSharedMemory> StructMemoryPtr;
diff --git a/Engine/Plugins/Experimental/StructUtils/Source/StructUtils/Public/StructUtils.h b/Engine/Plugins/Experimental/StructUtils/Source/StructUtils/Public/StructUtils.h
--- a/Engine/Plugins/Experimental/StructUtils/Source/StructUtils/Public/StructUtils.h
+++ b/Engine/Plugins/Experimental/StructUtils/Source/StructUtils/Public/StructUtils.h
@@ -17,6 +17,16 @@ namespace UE::StructUtils
 					  !TIsDerivedFrom<T, struct FConstSharedStruct>::IsDerived, "It does not make sense to create wrapper over these types.");
 	}
 
+	template <typename T>
+	void CheckWrapperType()
+	{
+		static_assert(TIsDerivedFrom<T, struct FInstancedStruct>::IsDerived
+			|| TIsDerivedFrom<T, struct FStructView>::IsDerived
+			|| TIsDerivedFrom<T, struct FConstStructView>::IsDerived
+			|| TIsDerivedFrom<T, struct FSharedStruct>::IsDerived
+			|| TIsDerivedFrom<T, struct FConstSharedStruct>::IsDerived, "This function is meant to compare contents of wrapped structs.");
+	}
+
 	/** Returns reference to the struct, this assumes that all data is valid. */
 	template<typename T>
 	T& GetStructRef(const UScriptStruct* ScriptStruct, void* StructMemory)
