Added an overloaded version of `ExportTypes` that can use a callback.

The new version of `ExportTypes` takes a function reference instead of an array. This avoids the need for an intermediate array for storage and can simplify code if tags and fragments both have the same functionality to process.

diff --git a/Engine/Plugins/Experimental/StructUtils/Source/StructUtils/Public/StructTypeBitSet.h b/Engine/Plugins/Experimental/StructUtils/Source/StructUtils/Public/StructTypeBitSet.h
--- a/Engine/Plugins/Experimental/StructUtils/Source/StructUtils/Public/StructTypeBitSet.h
+++ b/Engine/Plugins/Experimental/StructUtils/Source/StructUtils/Public/StructTypeBitSet.h
@@ -5,6 +5,7 @@
 #include "Containers/BitArray.h"
 #include "StructUtilsTypes.h"
 #include "InstancedStruct.h"
+#include "Templates/FunctionFwd.h"
 
 
 class FArchive;
@@ -566,6 +567,27 @@ public:
 		}
 	}
 
+	/**
+	 * Lists all types used by this bit set, calling the provided callback for each one. Returning false from
+	 * the callback will early-out of iterating over the types.
+	 * 
+	 * note that this function is slow(ish) due to the FStructTracker utilizing WeakObjectPtrs to store types.
+	 * @todo To be improved.
+	 */
+	void ExportTypes(TFunctionRef<bool(const TUStructType*)> Callback) const
+	{
+		TBitArray<>::FConstIterator It(StructTypesBitArray);
+		bool bKeepGoing = true;
+		while (bKeepGoing && It)
+		{
+			if (It.GetValue())
+			{
+				bKeepGoing = Callback(GetTypeAtIndex(It.GetIndex()));
+			}
+			++It;
+		}
+	}
+
 	SIZE_T GetAllocatedSize() const
 	{
 		return StructTypesBitArray.GetAllocatedSize();
