[StructUtils] Added templated types for StructView, SharedStruct

diff --git a/Engine/Plugins/Experimental/StructUtils/Source/StructUtils/Public/InstancedStruct.h b/Engine/Plugins/Experimental/StructUtils/Source/StructUtils/Public/InstancedStruct.h
--- a/Engine/Plugins/Experimental/StructUtils/Source/StructUtils/Public/InstancedStruct.h
+++ b/Engine/Plugins/Experimental/StructUtils/Source/StructUtils/Public/InstancedStruct.h
@@ -7,7 +7,9 @@
 #include "InstancedStruct.generated.h"
 
 struct FConstStructView;
+template<typename BaseStructT> struct TConstStructView;
 struct FConstSharedStruct;
+template<typename BaseStructT> struct TConstSharedStruct;
 class UUserDefinedStruct;
 
 /**
@@ -296,6 +298,16 @@ struct TInstancedStruct
 public:
 	TInstancedStruct() = default;
 
+	/**
+	 * This constructor is explicit to avoid accidentally converting struct views to instanced structs (which would result in costly copy of the struct to be made).
+	 * Implicit conversion could happen e.g. when comparing TInstancedStruct to TConstStructView.
+	 */
+	template<typename T = BaseStructT, typename = std::enable_if_t<std::is_base_of_v<BaseStructT, std::decay_t<T>>>>
+	explicit TInstancedStruct(const TConstStructView<T> InOther)
+	{
+		InitializeAsScriptStruct(InOther.GetScriptStruct(), InOther.GetMemory());
+	}
+
 	template<typename T = BaseStructT, typename = std::enable_if_t<std::is_base_of_v<BaseStructT, std::decay_t<T>>>>
 	TInstancedStruct(const TInstancedStruct<T>& InOther)
 		: InstancedStruct(InOther.InstancedStruct)
@@ -308,6 +320,16 @@ public:
 	{
 	}
 
+	template<typename T = BaseStructT, typename = std::enable_if_t<std::is_base_of_v<BaseStructT, std::decay_t<T>>>>
+	TInstancedStruct& operator=(const TConstStructView<T> InOther)
+	{
+		if (this != &InOther)
+		{
+			InstancedStruct = InOther.InstancedStruct;
+		}
+		return *this;
+	}
+
 	template<typename T = BaseStructT, typename = std::enable_if_t<std::is_base_of_v<BaseStructT, std::decay_t<T>>>>
 	TInstancedStruct& operator=(const TInstancedStruct<T>& InOther)
 	{
diff --git a/Engine/Plugins/Experimental/StructUtils/Source/StructUtils/Public/SharedStruct.h b/Engine/Plugins/Experimental/StructUtils/Source/StructUtils/Public/SharedStruct.h
--- a/Engine/Plugins/Experimental/StructUtils/Source/StructUtils/Public/SharedStruct.h
+++ b/Engine/Plugins/Experimental/StructUtils/Source/StructUtils/Public/SharedStruct.h
@@ -2,12 +2,17 @@
 
 #pragma once
 
+#include "InstancedStruct.h"
 #include "StructUtils.h"
 #include "StructUtilsTypes.h"
+
 #include "SharedStruct.generated.h"
 
 struct FInstancedStruct;
+template<typename BaseStructT> struct TInstancedStruct;
+template<typename BaseStructT> struct TStructView;
 struct FConstStructView;
+template<typename BaseStructT> struct TConstStructView;
 
 ///////////////////////////////////////////////////////////////// FStructSharedMemory /////////////////////////////////////////////////////////////////
 
@@ -339,6 +344,174 @@ protected:
 	TSharedPtr<FStructSharedMemory> StructMemoryPtr;
 };
 
+/**
+ * TSharedStruct is a type-safe FSharedStruct wrapper against the given BaseStruct type.
+ * @note When used as a property, this automatically defines the BaseStruct property meta-data.
+ * 
+ * Example:
+ *
+ *	TSharedStruct<FTestStructBase> Test;
+ *
+ *	TArray<TSharedStruct<FTestStructBase>> TestArray;
+ */
+template<typename BaseStructT>
+struct TSharedStruct
+{
+	explicit TSharedStruct() = default;
+
+	/** Copy constructors */
+	template<typename T = BaseStructT, typename = std::enable_if_t<std::is_base_of_v<BaseStructT, std::decay_t<T>>>>
+	TSharedStruct(const TSharedStruct<T>& InOther)
+		: SharedStruct(InOther.SharedStruct)
+	{
+	}
+
+	template<typename T = BaseStructT, typename = std::enable_if_t<std::is_base_of_v<BaseStructT, std::decay_t<T>>>>
+	TSharedStruct(TSharedStruct<T>&& InOther)
+		: SharedStruct(MoveTemp(InOther.SharedStruct))
+	{
+	}
+
+	/** Assignment operators */
+	template<typename T = BaseStructT, typename = std::enable_if_t<std::is_base_of_v<BaseStructT, std::decay_t<T>>>>
+	TSharedStruct& operator=(const TSharedStruct<T>& InOther)
+	{
+		if (this != &InOther)
+		{
+			SharedStruct = InOther.SharedStruct;
+		}
+		return *this;
+	}
+
+	template<typename T = BaseStructT, typename = std::enable_if_t<std::is_base_of_v<BaseStructT, std::decay_t<T>>>>
+	TSharedStruct& operator=(TSharedStruct<T>&& InOther)
+	{
+		if (this != &InOther)
+		{
+			SharedStruct = MoveTemp(InOther.SharedStruct);
+		}
+		return *this;
+	}
+
+	/** Returns struct type. */
+	const UScriptStruct* GetScriptStruct() const
+	{
+		return SharedStruct.GetScriptStruct();
+	}
+
+	TObjectPtr<const UScriptStruct>* const GetScriptStructPtr() const
+	{
+		return SharedStruct.GetScriptStructPtr();
+	}
+
+	/** Returns a mutable pointer to struct memory. */
+	uint8* GetMemory() const
+	{
+		return SharedStruct.GetMemory();
+	}
+
+	/** Reset to empty. */
+	void Reset()
+	{
+		SharedStruct.Reset();
+	}
+
+	/** Initializes from templated struct type. This will create a new instance of the shared struct memory. */
+	template<typename T = BaseStructT, typename = std::enable_if_t<std::is_base_of_v<BaseStructT, std::decay_t<T>>>>
+	void Initialize()
+	{
+		SharedStruct.InitializeAs<T>();
+	}
+
+	/** Initializes from templated struct instance. This will create a new instance of the shared struct memory. */
+	template<typename T = BaseStructT,
+		typename = std::enable_if_t<
+			std::is_base_of_v<BaseStructT, std::decay_t<T>>
+			&& !(std::is_same_v<TStructView<T>, T>
+				|| std::is_same_v<TConstStructView<T>, T>
+				|| std::is_same_v<TSharedStruct<T>, T>
+				|| std::is_same_v<TConstSharedStruct<T>, T>
+				|| std::is_same_v<TInstancedStruct<T>, T>), T>>
+	void Initialize(const T& Struct)
+	{
+		SharedStruct.InitializeAs(TBaseStructure<T>::Get(), reinterpret_cast<const uint8*>(&Struct));
+	}
+
+	/** Initializes from struct type and emplace args. This will create a new instance of the shared struct memory.*/
+	template<typename T = BaseStructT, typename... TArgs, typename = std::enable_if_t<std::is_base_of_v<BaseStructT, std::decay_t<T>>>>
+	void Initialize(TArgs&&... InArgs)
+	{
+		SharedStruct.InitializeAs(Forward<TArgs>(InArgs)...);
+	}
+
+	/** Creates a new TSharedStruct. This will create a new instance of the shared struct memory. */
+	template<typename T = BaseStructT, typename = std::enable_if_t<std::is_base_of_v<BaseStructT, std::decay_t<T>>>>
+	static TSharedStruct Make()
+	{
+		return FSharedStruct::Make<T>();
+	}
+
+	/** Creates a new TSharedStruct from templated struct instance. This will create a new instance of the shared struct memory. */
+	template<typename T = BaseStructT, typename = std::enable_if_t<std::is_base_of_v<BaseStructT, std::decay_t<T>>>>
+	static TSharedStruct Make(const T& Struct)
+	{
+		return FSharedStruct::Make<T>(Struct);
+	}
+
+	/** Creates a new TSharedStruct from the templated type and forward all arguments to constructor. This will create a new instance of the shared struct memory. */
+	template< typename... TArgs>
+	static TSharedStruct Make(TArgs&&... InArgs)
+	{
+		return FSharedStruct::Make<BaseStructT>(Forward<TArgs>(InArgs)...);
+	}
+
+	/** Returns reference to the struct, this getter assumes that all data is valid. */
+	template<typename T = BaseStructT, typename = std::enable_if_t<std::is_base_of_v<BaseStructT, std::decay_t<T>>>>
+	T& Get() const
+	{
+		return SharedStruct.Get<T>();
+	}
+
+	/** Returns pointer to the struct, or nullptr if cast is not valid. */
+	template<typename T = BaseStructT, typename = std::enable_if_t<std::is_base_of_v<BaseStructT, std::decay_t<T>>>>
+	T* GetPtr() const
+	{
+		return SharedStruct.GetPtr<T>();
+	}
+
+	/** Returns True if the struct is valid.*/
+	bool IsValid() const
+	{
+		return SharedStruct.IsValid();
+	}
+
+	/** Comparison operators. Note: it does not compare the internal structure itself */
+	template<typename T = BaseStructT, typename = std::enable_if_t<std::is_base_of_v<BaseStructT, std::decay_t<T>>>>
+	bool operator==(const T& Other) const
+	{
+		return SharedStruct.operator==(Other);
+	}
+
+	template<typename T = BaseStructT, typename = std::enable_if_t<std::is_base_of_v<BaseStructT, std::decay_t<T>>>>
+	bool operator!=(const T& Other) const
+	{
+		return SharedStruct.operator!=(Other);
+	}
+
+private:
+	/**
+	 * Note:
+	 *   TSharedStruct is a wrapper for a FSharedStruct (rather than inheriting) so that it can provide a locked-down type-safe 
+	 *   API for use in C++, without being able to accidentally take a reference to the untyped API to workaround the restrictions.
+	 * 
+	 *   TSharedStruct MUST be the same size as FSharedStruct, as the reflection layer treats a TSharedStruct as a FSharedStruct.
+	 *   This means that any reflected APIs (like ExportText) that accept an FSharedStruct pointer can also accept a TSharedStruct pointer.
+	 */
+	FSharedStruct SharedStruct;
+
+	template <typename U> friend struct TSharedStruct;
+};
+
 template<>
 struct TStructOpsTypeTraits<FSharedStruct> : public TStructOpsTypeTraitsBase2<FSharedStruct>
 {
@@ -564,6 +737,197 @@ protected:
 	TSharedPtr<const FStructSharedMemory> StructMemoryPtr;
 };
 
+/**
+ * TConstSharedStruct is a type-safe FConstSharedStruct wrapper against the given BaseStruct type.
+ * @note When used as a property, this automatically defines the BaseStruct property meta-data.
+ * 
+ * Example:
+ *
+ *	TConstSharedStruct<FTestStructBase> Test;
+ *
+ *	TArray<TConstSharedStruct<FTestStructBase>> TestArray;
+ */
+template<typename BaseStructT>
+struct TConstSharedStruct
+{
+	explicit TConstSharedStruct() = default;
+
+	template<typename T = BaseStructT, typename = std::enable_if_t<std::is_base_of_v<BaseStructT, std::decay_t<T>>>>
+	TConstSharedStruct(const TConstSharedStruct<T>& Other)
+		: ConstSharedStruct(Other.ConstSharedStruct)
+	{}
+
+	template<typename T = BaseStructT, typename = std::enable_if_t<std::is_base_of_v<BaseStructT, std::decay_t<T>>>>
+	TConstSharedStruct(const TSharedStruct<T>& ConstSharedStruct)
+		: ConstSharedStruct(ConstSharedStruct.SharedStruct)
+	{}
+
+	template<typename T = BaseStructT, typename = std::enable_if_t<std::is_base_of_v<BaseStructT, std::decay_t<T>>>>
+	TConstSharedStruct(TConstSharedStruct<T>&& Other)
+		: ConstSharedStruct(MoveTemp(Other.ConstSharedStruct))
+	{}
+
+	template<typename T = BaseStructT, typename = std::enable_if_t<std::is_base_of_v<BaseStructT, std::decay_t<T>>>>
+	TConstSharedStruct(TSharedStruct<T>&& ConstSharedStruct)
+		: ConstSharedStruct(MoveTemp(ConstSharedStruct.SharedStruct))
+	{}
+
+	template<typename T = BaseStructT, typename = std::enable_if_t<std::is_base_of_v<BaseStructT, std::decay_t<T>>>>
+	TConstSharedStruct& operator=(const TConstSharedStruct<T>& Other)
+	{
+		if (this == &Other)
+		{
+			ConstSharedStruct = Other.ConstSharedStruct;
+		}
+		return *this;
+	}
+
+	template<typename T = BaseStructT, typename = std::enable_if_t<std::is_base_of_v<BaseStructT, std::decay_t<T>>>>
+	TConstSharedStruct& operator=(const TSharedStruct<T>& SharedStruct)
+	{
+		if (this != &SharedStruct)
+		{
+			ConstSharedStruct = SharedStruct.SharedStruct;
+		}
+		return *this;
+	}
+
+	/** Returns const pointer to struct memory. */
+	const uint8* GetMemory() const
+	{
+		return ConstSharedStruct.GetMemory();
+	}
+
+	/** Reset to empty. */
+	void Reset()
+	{
+		ConstSharedStruct.Reset();
+	}
+
+	/** Initializes from templated struct type. */
+	void Initialize()
+	{
+		ConstSharedStruct.InitializeAs<BaseStructT>();
+	}
+
+	/** Initializes from templated struct instance. This will create a new instance of the shared struct memory. */
+	template<typename T = BaseStructT,
+		typename = std::enable_if_t<
+			std::is_base_of_v<BaseStructT, std::decay_t<T>>
+			&& !(std::is_same_v<TStructView<T>, T>
+				|| std::is_same_v<TConstStructView<T>, T>
+				|| std::is_same_v<TSharedStruct<T>, T>
+				|| std::is_same_v<TConstSharedStruct<T>, T>
+				|| std::is_same_v<TInstancedStruct<T>, T>), T>>
+	void Initialize(const T& Struct)
+	{
+		ConstSharedStruct.InitializeAs<T>(Struct);
+	}
+
+	/** Initializes from struct type and emplace args. This will create a new instance of the shared struct memory. */
+	template<typename T = BaseStructT, typename... TArgs, typename = std::enable_if_t<std::is_base_of_v<BaseStructT, std::decay_t<T>>>>
+	void Initialize(TArgs&&... InArgs)
+	{
+		ConstSharedStruct.InitializeAs<T>(Forward<TArgs>(InArgs)...);
+	}
+
+	/** Creates a new TSharedStruct from templated struct type. This will create a new instance of the shared struct memory. */
+	template<typename T = BaseStructT, typename = std::enable_if_t<std::is_base_of_v<BaseStructT, std::decay_t<T>>>>
+	static TConstSharedStruct Make()
+	{
+		FConstSharedStruct SharedStruct;
+		SharedStruct.InitializeAs<BaseStructT>();
+		return SharedStruct;
+	}
+
+	/** Creates a new TSharedStruct from templated struct instance. This will create a new instance of the shared struct memory. */
+	template<typename T = BaseStructT, typename = std::enable_if_t<std::is_base_of_v<BaseStructT, std::decay_t<T>>>>
+	static TConstSharedStruct Make(const T& Struct)
+	{
+		FConstSharedStruct SharedStruct;
+		SharedStruct.InitializeAs<BaseStructT>(Struct);
+		return SharedStruct;
+	}
+
+	/** Creates a new TSharedStruct from struct type and optional data. This will create a new instance of the shared struct memory. */
+	template<typename T = BaseStructT, typename = std::enable_if_t<std::is_base_of_v<BaseStructT, std::decay_t<T>>>>
+	static TConstSharedStruct Make(const uint8* InStructMemory = nullptr)
+	{
+		TConstSharedStruct SharedStruct;
+		SharedStruct.template Initialize<T>(InStructMemory);
+		return SharedStruct;
+	}
+
+	/** Creates a new TSharedStruct from the templated type and forward all arguments to constructor. This will create a new instance of the shared struct memory. */
+	template<typename T = BaseStructT, typename... TArgs, typename = std::enable_if_t<std::is_base_of_v<BaseStructT, std::decay_t<T>>>>
+	static TConstSharedStruct Make(TArgs&&... InArgs)
+	{
+		TConstSharedStruct SharedStruct;
+		SharedStruct.template Initialize<T>(Forward<TArgs>(InArgs)...);
+		return SharedStruct;
+	}
+
+	/** Returns const reference to the struct, this getter assumes that all data is valid. */
+	template<typename T = BaseStructT, typename = std::enable_if_t<std::is_base_of_v<BaseStructT, std::decay_t<T>>>>
+	constexpr const T& Get() const
+	{
+		return ConstSharedStruct.Get<T>();
+	}
+
+	/** Returns const reference to the struct, this getter assumes that all data is valid. */
+	template<typename T = BaseStructT, typename = std::enable_if_t<std::is_base_of_v<BaseStructT, std::decay_t<T>>>>
+	constexpr typename TEnableIf<TIsConst<T>::Value, T&>::Type Get() const
+	{
+		return ConstSharedStruct.Get<T>();
+	}
+
+	/** Returns const pointer to the struct. */
+	template<typename T = BaseStructT, typename = std::enable_if_t<std::is_base_of_v<BaseStructT, std::decay_t<T>>>>
+	constexpr const T* GetPtr() const
+	{
+		return ConstSharedStruct.GetPtr<T>();
+	}
+
+	/** Returns const pointer to the struct, or nullptr if cast is not valid. */
+	template<typename T = BaseStructT, typename = std::enable_if_t<std::is_base_of_v<BaseStructT, std::decay_t<T>>>>
+	constexpr typename TEnableIf<TIsConst<T>::Value, T*>::Type GetPtr() const
+	{
+		return ConstSharedStruct.GetPtr<T>();
+	}
+
+	/** Returns True if the struct is valid.*/
+	bool IsValid() const
+	{
+		return ConstSharedStruct.IsValid();
+	}
+
+	/** Comparison operators. Note: it does not compare the internal structure itself */
+	template<typename T = BaseStructT, typename = std::enable_if_t<std::is_base_of_v<BaseStructT, std::decay_t<T>>>>
+	bool operator==(const T& Other) const
+	{
+		return ConstSharedStruct.operator==(Other);
+	}
+
+	template<typename T = BaseStructT, typename = std::enable_if_t<std::is_base_of_v<BaseStructT, std::decay_t<T>>>>
+	bool operator!=(const T& Other) const
+	{
+		return ConstSharedStruct.operator!=(Other);
+	}
+
+private:
+	/**
+	 * Note:
+	 *   TConstSharedStruct is a wrapper for a FConstSharedStruct (rather than inheriting) so that it can provide a locked-down type-safe 
+	 *   API for use in C++, without being able to accidentally take a reference to the untyped API to workaround the restrictions.
+	 * 
+	 *   TConstSharedStruct MUST be the same size as FConstSharedStruct, as the reflection layer treats a TConstSharedStruct as a FConstSharedStruct.
+	 *   This means that any reflected APIs (like ExportText) that accept an FConstSharedStruct pointer can also accept a TConstSharedStruct pointer.
+	 */
+	FConstSharedStruct ConstSharedStruct;
+
+	template <typename U> friend struct TConstSharedStruct;
+};
+
 template<>
 struct TStructOpsTypeTraits<FConstSharedStruct> : public TStructOpsTypeTraitsBase2<FConstSharedStruct>
 {
diff --git a/Engine/Plugins/Experimental/StructUtils/Source/StructUtils/Public/StructView.h b/Engine/Plugins/Experimental/StructUtils/Source/StructUtils/Public/StructView.h
--- a/Engine/Plugins/Experimental/StructUtils/Source/StructUtils/Public/StructView.h
+++ b/Engine/Plugins/Experimental/StructUtils/Source/StructUtils/Public/StructView.h
@@ -126,6 +126,70 @@ protected:
 	uint8* StructMemory = nullptr;
 };
 
+/**
+ * TStructView is a type-safe FStructView wrapper against the given BaseStruct type.
+ * 
+ * Example:
+ *
+ *	TStructView<FTestStructBase> Test;
+ *
+ *	TArray<TStructView<FTestStructBase>> TestArray;
+ */
+template<typename BaseStructT>
+struct TStructView : FStructView
+{
+public:
+
+	explicit TStructView() = default;
+
+	template<typename T = BaseStructT, typename = std::enable_if_t<std::is_base_of_v<BaseStructT, std::decay_t<T>>>>
+	TStructView(BaseStructT& InStruct)
+		: FStructView(TBaseStructure<BaseStructT>::Get(), reinterpret_cast<uint8*>(&InStruct))
+	{}
+
+	template<typename T = BaseStructT, typename = std::enable_if_t<std::is_base_of_v<BaseStructT, std::decay_t<T>>>>
+	explicit TStructView(uint8* InStructMemory = nullptr)
+		: FStructView(T::StaticStruct(), InStructMemory)
+	{}
+
+	template<typename T = BaseStructT, typename = std::enable_if_t<std::is_base_of_v<BaseStructT, std::decay_t<T>>>>
+	TStructView(TInstancedStruct<T>& InstancedStruct)
+		: TStructView(InstancedStruct.GetScriptStruct(), InstancedStruct.GetMutableMemory())
+	{}
+
+	template<typename T = BaseStructT, typename = std::enable_if_t<std::is_base_of_v<BaseStructT, std::decay_t<T>>>>
+	TStructView(const TSharedStruct<T>& SharedStruct)
+		: TStructView(SharedStruct.GetScriptStruct(), SharedStruct.GetMemory())
+	{}
+
+	/** Returns mutable reference to the struct, this getter assumes that all data is valid. */
+	template<typename T = BaseStructT, typename = std::enable_if_t<std::is_base_of_v<BaseStructT, std::decay_t<T>>>>
+	T& Get() const
+	{
+		return UE::StructUtils::GetStructRef<BaseStructT>(ScriptStruct, StructMemory);
+	}
+
+	/** Returns mutable pointer to the struct, or nullptr if cast is not valid. */
+	template<typename T = BaseStructT, typename = std::enable_if_t<std::is_base_of_v<BaseStructT, std::decay_t<T>>>>
+	T* GetPtr() const
+	{
+		return UE::StructUtils::GetStructPtr<T>(ScriptStruct, StructMemory);
+	}
+
+	/** Comparison operators. Note: it does not compare the internal structure itself */
+	template<typename T = BaseStructT, typename = std::enable_if_t<std::is_base_of_v<BaseStructT, std::decay_t<T>>>>
+	bool operator==(const TStructView<T>& Other) const
+	{
+		return ((ScriptStruct == Other.GetScriptStruct()) && (StructMemory == Other.GetMemory()));
+	}
+
+	template<typename T = BaseStructT, typename = std::enable_if_t<std::is_base_of_v<BaseStructT, std::decay_t<T>>>>
+	bool operator!=(const TStructView<T>& Other) const
+	{
+		return !operator==(Other);
+	}
+};
+
 ///////////////////////////////////////////////////////////////// FConstStructView /////////////////////////////////////////////////////////////////
 /**
  * FConstStructView is "typed" struct pointer, it contains const pointer to struct plus UScriptStruct pointer.
@@ -258,3 +322,75 @@ protected:
 	const UScriptStruct* ScriptStruct = nullptr;
 	const uint8* StructMemory = nullptr;
 };
+
+/**
+ * TConstStructView is a type-safe FConstStructView wrapper against the given BaseStruct type.
+ * 
+ * Example:
+ *
+ *	TConstStructView<FTestStructBase> Test;
+ *
+ *	TArray<TConstStructView<FTestStructBase>> TestArray;
+ */
+template<typename BaseStructT>
+struct TConstStructView : FConstStructView
+{
+	explicit TConstStructView() = default;
+
+	template<typename T = BaseStructT, typename = std::enable_if_t<std::is_base_of_v<BaseStructT, std::decay_t<T>>>>
+	TConstStructView(const BaseStructT& InStruct)
+		: FConstStructView(TBaseStructure<T>::Get(), reinterpret_cast<const uint8*>(&InStruct))
+	{}
+
+	template<typename T = BaseStructT, typename = std::enable_if_t<std::is_base_of_v<BaseStructT, std::decay_t<T>>>>
+	explicit TConstStructView(const uint8* InStructMemory = nullptr)
+		: FConstStructView(T::StaticStruct(), InStructMemory)
+	{}
+
+	template<typename T = BaseStructT, typename = std::enable_if_t<std::is_base_of_v<BaseStructT, std::decay_t<T>>>>
+	TConstStructView(const TInstancedStruct<T>& InstancedStruct)
+		: FConstStructView(InstancedStruct.GetScriptStruct(), InstancedStruct.GetMemory())
+	{}
+
+	template<typename T = BaseStructT, typename = std::enable_if_t<std::is_base_of_v<BaseStructT, std::decay_t<T>>>>
+	TConstStructView(const TSharedStruct<T>& SharedStruct)
+		: FConstStructView(SharedStruct.GetScriptStruct(), SharedStruct.GetMemory())
+	{}
+
+	template<typename T = BaseStructT, typename = std::enable_if_t<std::is_base_of_v<BaseStructT, std::decay_t<T>>>>
+	TConstStructView(const TConstSharedStruct<T>& SharedStruct)
+		: FConstStructView(SharedStruct.GetScriptStruct(), SharedStruct.GetMemory())
+	{}
+
+	template<typename T = BaseStructT, typename = std::enable_if_t<std::is_base_of_v<BaseStructT, std::decay_t<T>>>>
+	TConstStructView(const TStructView<T> StructView)
+		: FConstStructView(StructView.GetScriptStruct(), StructView.GetMemory())
+	{}
+
+	/** Returns const reference to the struct, this getter assumes that all data is valid. */
+	template<typename T = BaseStructT, typename = std::enable_if_t<std::is_base_of_v<BaseStructT, std::decay_t<T>>>>
+	constexpr T& Get() const
+	{
+		return UE::StructUtils::GetStructRef<T>(ScriptStruct, StructMemory);
+	}
+
+	/** Returns const pointer to the struct, or nullptr if cast is not valid. */
+	template<typename T = BaseStructT, typename = std::enable_if_t<std::is_base_of_v<BaseStructT, std::decay_t<T>>>>
+	constexpr T* GetPtr() const
+	{
+		return UE::StructUtils::GetStructPtr<T>(ScriptStruct, StructMemory);
+	}
+
+	/** Comparison operators. Note: it does not compare the internal structure itself */
+	template<typename T = BaseStructT, typename = std::enable_if_t<std::is_base_of_v<BaseStructT, std::decay_t<T>>>>
+	bool operator==(const TConstStructView<T>& Other) const
+	{
+		return ((ScriptStruct == Other.GetScriptStruct()) && (StructMemory == Other.GetMemory()));
+	}
+
+	template<typename T = BaseStructT, typename = std::enable_if_t<std::is_base_of_v<BaseStructT, std::decay_t<T>>>>
+	bool operator!=(const TConstStructView<T>& Other) const
+	{
+		return !operator==(Other);
+	}
+};
diff --git a/Engine/Plugins/Experimental/StructUtils/Source/StructUtilsTestSuite/Private/InstancedStructTest.cpp b/Engine/Plugins/Experimental/StructUtils/Source/StructUtilsTestSuite/Private/InstancedStructTest.cpp
--- a/Engine/Plugins/Experimental/StructUtils/Source/StructUtilsTestSuite/Private/InstancedStructTest.cpp
+++ b/Engine/Plugins/Experimental/StructUtils/Source/StructUtilsTestSuite/Private/InstancedStructTest.cpp
@@ -1,4 +1,4 @@
-﻿// Copyright Epic Games, Inc. All Rights Reserved.
+// Copyright Epic Games, Inc. All Rights Reserved.
 
 #include "AITestsCommon.h"
 #include "SharedStruct.h"
@@ -37,6 +37,14 @@ namespace FInstancedStructTest
 
 				AITEST_EQUAL("FInstancedStruct initialized from Make should have value of FStructView its initiliazed from", Val, InstancedStruct.Get<FTestStructSimple>().Float);
 			}
+			
+			{
+				FTestStructSimple Simple(Val);
+				TConstStructView<FTestStructSimple> ConstStructView(Simple);
+				TInstancedStruct<FTestStructSimple> InstancedStruct(ConstStructView);
+
+				AITEST_EQUAL("TInstancedStruct initialized from Make should have value of TConstStructView its initiliazed from", Val, InstancedStruct.Get<FTestStructSimple>().Float);
+			}
 
 			{
 				FTestStructSimple Simple(Val);
