Auto-expand Instanced Struct properties

diff --git a/Engine/Plugins/Experimental/StructUtils/Source/StructUtilsEditor/Private/Customizations/InstancedStructDetails.cpp b/Engine/Plugins/Experimental/StructUtils/Source/StructUtilsEditor/Private/Customizations/InstancedStructDetails.cpp
--- a/Engine/Plugins/Experimental/StructUtils/Source/StructUtilsEditor/Private/Customizations/InstancedStructDetails.cpp
+++ b/Engine/Plugins/Experimental/StructUtils/Source/StructUtilsEditor/Private/Customizations/InstancedStructDetails.cpp
@@ -431,6 +431,7 @@ void FInstancedStructDetails::CustomizeHeader(TSharedRef<class IPropertyHandle>
 	}
 
 	HeaderRow
+		.ShouldAutoExpand(true)
 		.NameContent()
 		[
 			StructPropertyHandle->CreatePropertyNameWidget()
