[Backout] - CL32829489

Original CL Desc
-----------------------------------------------------------------
Auto-expand Instanced Struct properties

diff --git a/Engine/Plugins/Experimental/StructUtils/Source/StructUtilsEditor/Private/Customizations/InstancedStructDetails.cpp b/Engine/Plugins/Experimental/StructUtils/Source/StructUtilsEditor/Private/Customizations/InstancedStructDetails.cpp
--- a/Engine/Plugins/Experimental/StructUtils/Source/StructUtilsEditor/Private/Customizations/InstancedStructDetails.cpp
+++ b/Engine/Plugins/Experimental/StructUtils/Source/StructUtilsEditor/Private/Customizations/InstancedStructDetails.cpp
@@ -431,7 +431,6 @@ void FInstancedStructDetails::CustomizeHeader(TSharedRef<class IPropertyHandle>
 	}
 
 	HeaderRow
-		.ShouldAutoExpand(true)
 		.NameContent()
 		[
 			StructPropertyHandle->CreatePropertyNameWidget()
