Fix TInstancedStruct assignment so Base = Derived assignments work

Without this errors are generated:
error C2446: '!=': no conversion from 'TInstancedStruct<FDerived> *' to 'TInstancedStruct<FBase> *'

With code like:
TInstancedStruct<FBase> Base;
TInstancedStruct<FDerived> Derived;
Base = Derived;   // Error C2446

FInstancedStruct::operator= already guards against assignment with this so this check is largely redundant anyway.

diff --git a/Engine/Plugins/Experimental/StructUtils/Source/StructUtils/Public/InstancedStruct.h b/Engine/Plugins/Experimental/StructUtils/Source/StructUtils/Public/InstancedStruct.h
--- a/Engine/Plugins/Experimental/StructUtils/Source/StructUtils/Public/InstancedStruct.h
+++ b/Engine/Plugins/Experimental/StructUtils/Source/StructUtils/Public/InstancedStruct.h
@@ -333,20 +333,14 @@ public:
 	template<typename T = BaseStructT, typename = std::enable_if_t<std::is_base_of_v<BaseStructT, std::decay_t<T>>>>
 	TInstancedStruct& operator=(const TInstancedStruct<T>& InOther)
 	{
-		if (this != &InOther)
-		{
-			InstancedStruct = InOther.InstancedStruct;
-		}
+		InstancedStruct = InOther.InstancedStruct;
 		return *this;
 	}
 
 	template<typename T = BaseStructT, typename = std::enable_if_t<std::is_base_of_v<BaseStructT, std::decay_t<T>>>>
 	TInstancedStruct& operator=(TInstancedStruct<T>&& InOther)
 	{
-		if (this != &InOther)
-		{
-			InstancedStruct = MoveTemp(InOther.InstancedStruct);
-		}
+		InstancedStruct = MoveTemp(InOther.InstancedStruct);
 		return *this;
 	}
 
