[StructUtils] const correctness for TConstStructView<>::Get()/GetPtr()

diff --git a/Engine/Plugins/Experimental/StructUtils/Source/StructUtils/Public/StructView.h b/Engine/Plugins/Experimental/StructUtils/Source/StructUtils/Public/StructView.h
--- a/Engine/Plugins/Experimental/StructUtils/Source/StructUtils/Public/StructView.h
+++ b/Engine/Plugins/Experimental/StructUtils/Source/StructUtils/Public/StructView.h
@@ -369,14 +369,14 @@ struct TConstStructView : FConstStructView
 
 	/** Returns const reference to the struct, this getter assumes that all data is valid. */
 	template<typename T = BaseStructT, typename = std::enable_if_t<std::is_base_of_v<BaseStructT, std::decay_t<T>>>>
-	constexpr T& Get() const
+	constexpr const T& Get() const
 	{
 		return UE::StructUtils::GetStructRef<T>(ScriptStruct, StructMemory);
 	}
 
 	/** Returns const pointer to the struct, or nullptr if cast is not valid. */
 	template<typename T = BaseStructT, typename = std::enable_if_t<std::is_base_of_v<BaseStructT, std::decay_t<T>>>>
-	constexpr T* GetPtr() const
+	constexpr const T* GetPtr() const
 	{
 		return UE::StructUtils::GetStructPtr<T>(ScriptStruct, StructMemory);
 	}
