Allow FInstancedStruct to import "()" as an alias for "None" to create an empty FInstancedStruct

"()" is used by UHT as a general "empty struct" marker

diff --git a/Engine/Plugins/Experimental/StructUtils/Source/StructUtils/Private/InstancedStruct.cpp b/Engine/Plugins/Experimental/StructUtils/Source/StructUtils/Private/InstancedStruct.cpp
--- a/Engine/Plugins/Experimental/StructUtils/Source/StructUtils/Private/InstancedStruct.cpp
+++ b/Engine/Plugins/Experimental/StructUtils/Source/StructUtils/Private/InstancedStruct.cpp
@@ -281,7 +281,12 @@ bool FInstancedStruct::ExportTextItem(FString& ValueStr, FInstancedStruct const&
 bool FInstancedStruct::ImportTextItem(const TCHAR*& Buffer, int32 PortFlags, UObject* Parent, FOutputDevice* ErrorText, FArchive* InSerializingArchive /*= nullptr*/)
 {
 	FNameBuilder StructPathName;
-	if (const TCHAR* Result = FPropertyHelpers::ReadToken(Buffer, StructPathName, /*bDottedNames*/true))
+	if (FCString::Strcmp(Buffer, TEXT("()")) == 0)
+	{
+		// UHT uses "()" as a general "empty struct" marker, so allow importing that as an alias for "None"
+		Buffer += 2;
+	}
+	else if (const TCHAR* Result = FPropertyHelpers::ReadToken(Buffer, StructPathName, /*bDottedNames*/true))
 	{
 		Buffer = Result;
 	}
