PropertyBag: Add the ability to add properties and their values in one operation

diff --git a/Engine/Plugins/Experimental/StructUtils/Source/StructUtils/Private/PropertyBag.cpp b/Engine/Plugins/Experimental/StructUtils/Source/StructUtils/Private/PropertyBag.cpp
--- a/Engine/Plugins/Experimental/StructUtils/Source/StructUtils/Private/PropertyBag.cpp
+++ b/Engine/Plugins/Experimental/StructUtils/Source/StructUtils/Private/PropertyBag.cpp
@@ -1672,6 +1672,38 @@ void FInstancedPropertyBag::AddProperty(const FName InName, const FProperty* InS
 	AddProperties({ FPropertyBagPropertyDesc(InName, InSourceProperty ) });
 }
 
+EPropertyBagResult FInstancedPropertyBag::ReplaceAllPropertiesAndValues(const TConstArrayView<FPropertyBagPropertyDesc> InDescs, const TConstArrayView<TConstArrayView<uint8>> InValues)
+{
+	if (InDescs.Num() != InValues.Num())
+	{
+		return EPropertyBagResult::OutOfBounds;
+	}
+
+	Reset();
+
+	const UPropertyBag* NewBagStruct = UPropertyBag::GetOrCreateFromDescs(InDescs);
+	InitializeFromBagStruct(NewBagStruct);
+
+	for (int32 Index = 0; Index < NewBagStruct->PropertyDescs.Num(); ++Index)
+	{
+		const FPropertyBagPropertyDesc& Desc = NewBagStruct->PropertyDescs[Index];
+		const TConstArrayView<uint8>& NewValue = InValues[Index];
+
+		if (NewValue.Num() == Desc.CachedProperty->GetSize())
+		{
+			void* TargetAddress = Value.GetMutableMemory() + Desc.CachedProperty->GetOffset_ForInternal();
+			void const* SourceAddress = NewValue.GetData();
+			Desc.CachedProperty->CopyCompleteValue(TargetAddress, SourceAddress);
+		}
+		else
+		{
+			return EPropertyBagResult::TypeMismatch;
+		}
+	}
+
+	return EPropertyBagResult::Success;
+}
+
 void FInstancedPropertyBag::RemovePropertiesByName(const TConstArrayView<FName> PropertiesToRemove)
 {
 	TArray<FPropertyBagPropertyDesc> Descs;
diff --git a/Engine/Plugins/Experimental/StructUtils/Source/StructUtils/Public/PropertyBag.h b/Engine/Plugins/Experimental/StructUtils/Source/StructUtils/Public/PropertyBag.h
--- a/Engine/Plugins/Experimental/StructUtils/Source/StructUtils/Public/PropertyBag.h
+++ b/Engine/Plugins/Experimental/StructUtils/Source/StructUtils/Public/PropertyBag.h
@@ -416,6 +416,14 @@ struct STRUCTUTILS_API FInstancedPropertyBag
 	 */
 	void AddProperty(const FName InName, const FProperty* InSourceProperty);
 
+	/**
+	 * Clears all properties, then adds the supplied properties to the bag and sets their respective values.
+	 * Numeric property values will be converted if possible, when a property's type changes.
+	 * @param InDescs Descriptors of new properties to add. Must be the same size as InValues.
+	 * @param InValues Values of new properties to add. Must be the same size as InDescs.
+	 */
+	EPropertyBagResult ReplaceAllPropertiesAndValues(const TConstArrayView<FPropertyBagPropertyDesc> InDescs, const TConstArrayView<TConstArrayView<uint8>> InValues);
+
 	/**
 	 * Removes properties from the bag by name if they exists.
 	 */
