PR #11841: StructUtils: Mark `FInstancedPropertyBag::GetValueSerializedString` as const

diff --git a/Engine/Plugins/Experimental/StructUtils/Source/StructUtils/Private/PropertyBag.cpp b/Engine/Plugins/Experimental/StructUtils/Source/StructUtils/Private/PropertyBag.cpp
--- a/Engine/Plugins/Experimental/StructUtils/Source/StructUtils/Private/PropertyBag.cpp
+++ b/Engine/Plugins/Experimental/StructUtils/Source/StructUtils/Private/PropertyBag.cpp
@@ -1992,7 +1992,7 @@ TValueOrError<FSoftObjectPath, EPropertyBagResult> FInstancedPropertyBag::GetVal
 	return MakeValue(std::move(ReturnValue));
 }
 
-TValueOrError<FString, EPropertyBagResult> FInstancedPropertyBag::GetValueSerializedString(const FName Name)
+TValueOrError<FString, EPropertyBagResult> FInstancedPropertyBag::GetValueSerializedString(const FName Name) const
 {
 	const FPropertyBagPropertyDesc* Desc = FindPropertyDescByName(Name);
 	if (Desc == nullptr)
diff --git a/Engine/Plugins/Experimental/StructUtils/Source/StructUtils/Public/PropertyBag.h b/Engine/Plugins/Experimental/StructUtils/Source/StructUtils/Public/PropertyBag.h
--- a/Engine/Plugins/Experimental/StructUtils/Source/StructUtils/Public/PropertyBag.h
+++ b/Engine/Plugins/Experimental/StructUtils/Source/StructUtils/Public/PropertyBag.h
@@ -496,7 +496,7 @@ struct STRUCTUTILS_API FInstancedPropertyBag
 	TValueOrError<FSoftObjectPath, EPropertyBagResult> GetValueSoftPath(const FName Name) const;
 
 	/** @return string-based serialized representation of the value. */
-	TValueOrError<FString, EPropertyBagResult> GetValueSerializedString(const FName Name);
+	TValueOrError<FString, EPropertyBagResult> GetValueSerializedString(const FName Name) const;
 
 	/** @return enum value of specified type. */
 	template <typename T>
