[StructUtils] added message to a check validating source and target struct types compatibility

diff --git a/Engine/Plugins/Experimental/StructUtils/Source/StructUtils/Public/StructUtils.h b/Engine/Plugins/Experimental/StructUtils/Source/StructUtils/Public/StructUtils.h
--- a/Engine/Plugins/Experimental/StructUtils/Source/StructUtils/Public/StructUtils.h
+++ b/Engine/Plugins/Experimental/StructUtils/Source/StructUtils/Public/StructUtils.h
@@ -33,7 +33,10 @@ namespace UE::StructUtils
 	{
 		check(StructMemory != nullptr);
 		check(ScriptStruct != nullptr);
-		check(ScriptStruct == TBaseStructure<T>::Get() || ScriptStruct->IsChildOf(TBaseStructure<T>::Get()));
+		checkf(ScriptStruct == TBaseStructure<T>::Get() || ScriptStruct->IsChildOf(TBaseStructure<T>::Get()),
+			TEXT("Incompatible struct types. Source type: '%s' Target type: '%s'"),
+			*(TBaseStructure<T>::Get()->GetName()),
+			ScriptStruct ? *(ScriptStruct->GetName()) : TEXT("nullptr"));
 		return *((T*)StructMemory);
 	}
 
