#UE Add some additional info to a log

diff --git a/Engine/Plugins/Experimental/StructUtils/Source/StructUtils/Private/InstancedStruct.cpp b/Engine/Plugins/Experimental/StructUtils/Source/StructUtils/Private/InstancedStruct.cpp
--- a/Engine/Plugins/Experimental/StructUtils/Source/StructUtils/Private/InstancedStruct.cpp
+++ b/Engine/Plugins/Experimental/StructUtils/Source/StructUtils/Private/InstancedStruct.cpp
@@ -202,7 +202,7 @@ bool FInstancedStruct::Serialize(FArchive& Ar)
 			// In this case we manually seek in the archive to skip its serialized content. 
 			// We don't want to rely on TaggedSerialization that will mark an error in the archive that
 			// may cause other serialization to fail (e.g. FArchive& operator<<(FArchive& Ar, TArray& A))
-			UE_LOG(LogCore, Warning, TEXT("Unable to find serialized UScriptStruct -> Advance %u bytes in the archive and reset to empty FInstancedStruct"), SerialSize);
+			UE_LOG(LogCore, Warning, TEXT("Unable to find serialized UScriptStruct -> Advance %u bytes in the archive and reset to empty FInstancedStruct. SerializedProperty:%s"), SerialSize, *Ar.GetSerializedProperty()->GetPathName());
 			Ar.Seek(Ar.Tell() + SerialSize);
 		}
 		else if (NonConstStruct != nullptr && ensureMsgf(GetMutableMemory() != nullptr, TEXT("A valid script struct should always have allocated memory")))
