StructUtils: Updated the InstancedStructDetails customization to take into account property category metadata. If properties have categories specified, they will be added under groups. Properties also need the "EnableCategories" metadata tag to opt-in to grouping.

diff --git a/Engine/Plugins/Experimental/StructUtils/Source/StructUtilsEditor/Private/Customizations/InstancedStructDetails.cpp b/Engine/Plugins/Experimental/StructUtils/Source/StructUtilsEditor/Private/Customizations/InstancedStructDetails.cpp
--- a/Engine/Plugins/Experimental/StructUtils/Source/StructUtilsEditor/Private/Customizations/InstancedStructDetails.cpp
+++ b/Engine/Plugins/Experimental/StructUtils/Source/StructUtilsEditor/Private/Customizations/InstancedStructDetails.cpp
@@ -7,6 +7,7 @@
 #include "Editor/AssetReferenceFilter.h"
 #include "Editor/EditorEngine.h"
 #include "IDetailChildrenBuilder.h"
+#include "IDetailGroup.h"
 #include "IPropertyUtilities.h"
 #include "UObject/Package.h"
 #include "Widgets/Images/SImage.h"
@@ -288,6 +289,80 @@ TArray<TWeakObjectPtr<const UStruct>> FInstancedStructDataDetails::GetInstanceTy
 	return Result;
 }
 
+void FInstancedStructDataDetails::GetPropertyGroups(const TArray<TSharedPtr<IPropertyHandle>>& InProperties, IDetailChildrenBuilder& InChildBuilder, TMap<TSharedPtr<IPropertyHandle>, IDetailGroup*>& OutPropertyToGroup) const
+{
+	static const FName CategoryName = FName(TEXT("Category"));
+	static const FName EnableCategoriesName = FName(TEXT("EnableCategories"));
+
+	// Temporarily store a mapping of category -> group while groups are being built
+	TMap<FString, IDetailGroup*> CategoryToGroup;
+	
+	for (const TSharedPtr<IPropertyHandle>& PropertyHandle : InProperties)
+	{
+		// The property needs the "EnableCategories" metadata in order to be added under a group. Grouping is opt-in.
+		if (!PropertyHandle->HasMetaData(EnableCategoriesName))
+		{
+			continue;
+		}
+		
+		const FString& PropertyCategory = PropertyHandle->GetMetaData(CategoryName).TrimStartAndEnd();
+		if (PropertyCategory.IsEmpty())
+		{
+			continue;
+		}
+
+		constexpr bool bCullEmpty = true;
+		TArray<FString> CategoriesToAdd;
+		PropertyCategory.ParseIntoArray(CategoriesToAdd, TEXT("|"), bCullEmpty);
+
+		if (CategoriesToAdd.IsEmpty())
+		{
+			continue;
+		}
+
+		// Tracks the category name as it is being built up (eg, Foo -> Foo|Bar -> Foo|Bar|Baz)
+		FString CompleteCategory;
+
+		// For this property, add all of the groups needed for its category (eg, Foo, Foo|Bar, and Foo|Bar|Baz)
+		IDetailGroup* CurrentGroup = nullptr;
+		for (FString& CategoryToAdd : CategoriesToAdd)
+		{
+			CategoryToAdd.TrimStartAndEndInline();
+
+			// Cover the edge case where there's a category like "Foo|", which is invalid
+			CategoryToAdd.TrimCharInline(TEXT('|'), nullptr);
+
+			CompleteCategory = CompleteCategory.IsEmpty()
+				? CategoryToAdd
+				: FString::Join(TArray({CompleteCategory, CategoryToAdd}), TEXT("|"));
+
+			// Create the category's group if it has not yet been created
+			if (!CategoryToGroup.Contains(CompleteCategory))
+			{
+				if (CurrentGroup)
+				{
+					// Add the group to the previous group if this is a nested category (eg, if this is the Foo|Bar group, add to the Foo group)
+					CurrentGroup = &CurrentGroup->AddGroup(FName(CompleteCategory), FText::FromString(CategoryToAdd));
+				}
+				else
+				{
+					// Otherwise, add the group as a normal group via the builder
+					CurrentGroup = &InChildBuilder.AddGroup(FName(CompleteCategory), FText::FromString(CategoryToAdd));
+				}
+
+				CategoryToGroup.Add(CompleteCategory, CurrentGroup);
+			}
+			else
+			{
+				CurrentGroup = CategoryToGroup[CompleteCategory];
+			}
+		}
+
+		check(CurrentGroup);
+		OutPropertyToGroup.Add(PropertyHandle, CurrentGroup);
+	}
+}
+
 void FInstancedStructDataDetails::OnStructLayoutChanges()
 {
 	if (StructProvider.IsValid())
@@ -360,10 +435,25 @@ void FInstancedStructDataDetails::GenerateChildContent(IDetailChildrenBuilder& C
 	else
 	{
 		TArray<TSharedPtr<IPropertyHandle>> ChildProperties = StructProperty->AddChildStructure(NewStructProvider);
+		
+		// Properties may have Category metadata. If that's the case, they should be added under groups.
+		TMap<TSharedPtr<IPropertyHandle>, IDetailGroup*> PropertyToGroup;
+		GetPropertyGroups(ChildProperties, ChildBuilder, PropertyToGroup);
+		
 		for (TSharedPtr<IPropertyHandle> ChildHandle : ChildProperties)
 		{
-			IDetailPropertyRow& Row = ChildBuilder.AddProperty(ChildHandle.ToSharedRef());
-			OnChildRowAdded(Row);
+			// If the property has a group, add it under the group. Otherwise, just add it normally via the builder.
+			IDetailGroup** PropertyGroup = PropertyToGroup.Find(ChildHandle);
+			if (PropertyGroup && *PropertyGroup)
+			{
+				IDetailPropertyRow& Row = (*PropertyGroup)->AddPropertyRow(ChildHandle.ToSharedRef());
+				OnChildRowAdded(Row);
+			}
+			else
+			{
+				IDetailPropertyRow& Row = ChildBuilder.AddProperty(ChildHandle.ToSharedRef());
+				OnChildRowAdded(Row);
+			}
 		}
 	}
 
diff --git a/Engine/Plugins/Experimental/StructUtils/Source/StructUtilsEditor/Public/InstancedStructDetails.h b/Engine/Plugins/Experimental/StructUtils/Source/StructUtilsEditor/Public/InstancedStructDetails.h
--- a/Engine/Plugins/Experimental/StructUtils/Source/StructUtilsEditor/Public/InstancedStructDetails.h
+++ b/Engine/Plugins/Experimental/StructUtils/Source/StructUtilsEditor/Public/InstancedStructDetails.h
@@ -9,6 +9,7 @@
 
 class IAssetReferenceFilter;
 class IPropertyHandle;
+class IDetailGroup;
 class IDetailPropertyRow;
 class IPropertyHandle;
 class FStructOnScope;
@@ -60,6 +61,7 @@ private:
  * Expects property handle holding FInstancedStruct as input.
  * Can be used in a implementation of a IPropertyTypeCustomization CustomizeChildren() to display editable FInstancedStruct contents.
  * OnChildRowAdded() is called right after each property is added, which allows the property row to be customizable.
+ * Child properties will be grouped if they 1) have "Category" metadata, and 2) have the "EnableCategories" metadata tag.
  */
 class STRUCTUTILSEDITOR_API FInstancedStructDataDetails : public IDetailCustomNodeBuilder, public TSharedFromThis<FInstancedStructDataDetails>
 {
@@ -91,6 +93,14 @@ private:
 	/** Returns type of the instanced struct for each instance/object being edited. */
 	TArray<TWeakObjectPtr<const UStruct>> GetInstanceTypes() const;
 
+	/**
+	 * Adds groups for the specified properties. One group is created for each unique category (from property metadata) that the properties have.
+	 * If a category is pipe-separated (eg, Foo|Bar), one group is added for "Foo" and another one for "Foo|Bar". In the returned map, the key is the
+	 * property, and the value is the group. If the property doesn't have a group (category), then it will not have an entry in the map. Note that
+	 * the property must opt-in to grouping by specifying the "EnableCategories" metadata tag.
+	 */
+	void GetPropertyGroups(const TArray<TSharedPtr<IPropertyHandle>>& InProperties, IDetailChildrenBuilder& InChildBuilder, TMap<TSharedPtr<IPropertyHandle>, IDetailGroup*>& OutPropertyToGroup) const;
+
 	/** Cached instance types, used to invalidate the layout when types change. */
 	TArray<TWeakObjectPtr<const UStruct>> CachedInstanceTypes;
 	
