TInstancedStruct: Add Serialize method

diff --git a/Engine/Plugins/Experimental/StructUtils/Source/StructUtils/Public/InstancedStruct.h b/Engine/Plugins/Experimental/StructUtils/Source/StructUtils/Public/InstancedStruct.h
--- a/Engine/Plugins/Experimental/StructUtils/Source/StructUtils/Public/InstancedStruct.h
+++ b/Engine/Plugins/Experimental/StructUtils/Source/StructUtils/Public/InstancedStruct.h
@@ -462,6 +462,11 @@ public:
 		InstancedStruct.AddStructReferencedObjects(Collector);
 	}
 
+	bool Serialize(FArchive& Ar)
+	{
+		return InstancedStruct.Serialize(Ar);
+	}
+
 private:
 	/**
 	 * Note:
