StructUtils: Separate SInstancedStructPicker out of FInstancedStructDetails

diff --git a/Engine/Plugins/Experimental/StructUtils/Source/StructUtilsEditor/Private/Customizations/InstancedStructDetails.cpp b/Engine/Plugins/Experimental/StructUtils/Source/StructUtilsEditor/Private/Customizations/InstancedStructDetails.cpp
--- a/Engine/Plugins/Experimental/StructUtils/Source/StructUtilsEditor/Private/Customizations/InstancedStructDetails.cpp
+++ b/Engine/Plugins/Experimental/StructUtils/Source/StructUtilsEditor/Private/Customizations/InstancedStructDetails.cpp
@@ -4,25 +4,20 @@
 #include "DetailWidgetRow.h"
 #include "DetailLayoutBuilder.h"
 #include "Editor.h"
-#include "Editor/AssetReferenceFilter.h"
 #include "Editor/EditorEngine.h"
 #include "IDetailChildrenBuilder.h"
 #include "IDetailGroup.h"
 #include "IPropertyUtilities.h"
 #include "UObject/Package.h"
 #include "Widgets/Images/SImage.h"
-#include "Widgets/Input/SComboButton.h"
-#include "ScopedTransaction.h"
 #include "Modules/ModuleManager.h"
 #include "StructViewerModule.h"
-#include "Styling/SlateIconFinder.h"
 #include "Engine/UserDefinedStruct.h"
 #include "InstancedStruct.h"
-#include "Widgets/Layout/SBox.h"
 #include "IStructureDataProvider.h"
-#include "GameFramework/Actor.h"
 #include "Misc/ConfigCacheIni.h"
 #include "StructUtilsDelegates.h"
+#include "SInstancedStructPicker.h"
 
 #define LOCTEXT_NAMESPACE "StructUtilsEditor"
 
@@ -164,84 +159,6 @@ protected:
 
 ////////////////////////////////////
 
-bool FInstancedStructFilter::IsStructAllowed(const FStructViewerInitializationOptions& InInitOptions, const UScriptStruct* InStruct, TSharedRef<FStructViewerFilterFuncs> InFilterFuncs)
-{
-	if (InStruct->IsA<UUserDefinedStruct>())
-	{
-		return bAllowUserDefinedStructs;
-	}
-
-	if (InStruct == BaseStruct)
-	{
-		return bAllowBaseStruct;
-	}
-
-	if (InStruct->HasMetaData(TEXT("Hidden")))
-	{
-		return false;
-	}
-
-	if (AssetReferenceFilter.IsValid())
-	{
-		if (!AssetReferenceFilter->PassesFilter(FAssetData(InStruct)))
-		{
-			return false;
-		}
-	}
-
-	// Query the native struct to see if it has the correct parent type (if any)
-	return !BaseStruct || InStruct->IsChildOf(BaseStruct);
-}
-
-bool FInstancedStructFilter::IsUnloadedStructAllowed(const FStructViewerInitializationOptions& InInitOptions, const FSoftObjectPath& InStructPath, TSharedRef<FStructViewerFilterFuncs> InFilterFuncs)
-{
-	// User Defined Structs don't support inheritance, so only include them requested
-	return bAllowUserDefinedStructs;
-}
-
-////////////////////////////////////
-
-namespace UE::StructUtils::Private
-{
-
-FPropertyAccess::Result GetCommonScriptStruct(TSharedPtr<IPropertyHandle> StructProperty, const UScriptStruct*& OutCommonStruct)
-{
-	bool bHasResult = false;
-	bool bHasMultipleValues = false;
-	
-	StructProperty->EnumerateConstRawData([&OutCommonStruct, &bHasResult, &bHasMultipleValues](const void* RawData, const int32 /*DataIndex*/, const int32 /*NumDatas*/)
-	{
-		if (const FInstancedStruct* InstancedStruct = static_cast<const FInstancedStruct*>(RawData))
-		{
-			const UScriptStruct* Struct = InstancedStruct->GetScriptStruct();
-
-			if (!bHasResult)
-			{
-				OutCommonStruct = Struct;
-			}
-			else if (OutCommonStruct != Struct)
-			{
-				bHasMultipleValues = true;
-			}
-
-			bHasResult = true;
-		}
-
-		return true;
-	});
-
-	if (bHasMultipleValues)
-	{
-		return FPropertyAccess::MultipleValues;
-	}
-	
-	return bHasResult ? FPropertyAccess::Success : FPropertyAccess::Fail;
-}
-
-} // UE::StructUtils::Private
-
-////////////////////////////////////
-
 FInstancedStructDataDetails::FInstancedStructDataDetails(TSharedPtr<IPropertyHandle> InStructProperty)
 {
 #if DO_CHECK
@@ -496,29 +413,10 @@ FInstancedStructDetails::~FInstancedStructDetails()
 
 void FInstancedStructDetails::CustomizeHeader(TSharedRef<class IPropertyHandle> StructPropertyHandle, class FDetailWidgetRow& HeaderRow, IPropertyTypeCustomizationUtils& StructCustomizationUtils)
 {
-	static const FName NAME_BaseStruct = "BaseStruct";
-	static const FName NAME_StructTypeConst = "StructTypeConst";
-
-	PropUtils = StructCustomizationUtils.GetPropertyUtilities();
 	StructProperty = StructPropertyHandle;
 	PropUtils = StructCustomizationUtils.GetPropertyUtilities();
 
 	OnObjectsReinstancedHandle = FCoreUObjectDelegates::OnObjectsReinstanced.AddSP(this, &FInstancedStructDetails::OnObjectsReinstanced);
-	
-	const bool bEnableStructSelection = !StructProperty->HasMetaData(NAME_StructTypeConst);
-
-	BaseScriptStruct = nullptr;
-	{
-		const FString& BaseStructName = StructProperty->GetMetaData(NAME_BaseStruct);
-		if (!BaseStructName.IsEmpty())
-		{
-			BaseScriptStruct = UClass::TryFindTypeSlow<UScriptStruct>(BaseStructName);
-			if (!BaseScriptStruct)
-			{
-				BaseScriptStruct = LoadObject<UScriptStruct>(nullptr, *BaseStructName);
-			}
-		}
-	}
 
 	HeaderRow
 		.ShouldAutoExpand(true)
@@ -530,30 +428,7 @@ void FInstancedStructDetails::CustomizeHeader(TSharedRef<class IPropertyHandle>
 		.MinDesiredWidth(250.f)
 		.VAlign(VAlign_Center)
 		[
-			SAssignNew(ComboButton, SComboButton)
-			.OnGetMenuContent(this, &FInstancedStructDetails::GenerateStructPicker)
-			.ContentPadding(0)
-			.IsEnabled(bEnableStructSelection)
-			.ButtonContent()
-			[
-				SNew(SHorizontalBox)
-				+ SHorizontalBox::Slot()
-				.AutoWidth()
-				.VAlign(VAlign_Center)
-				.Padding(0.0f, 0.0f, 4.0f, 0.0f)
-				[
-					SNew(SImage)
-					.Image(this, &FInstancedStructDetails::GetDisplayValueIcon)
-				]
-				+ SHorizontalBox::Slot()
-				.VAlign(VAlign_Center)
-				[
-					SNew(STextBlock)
-					.Text(this, &FInstancedStructDetails::GetDisplayValueString)
-					.ToolTipText(this, &FInstancedStructDetails::GetTooltipText)
-					.Font(IDetailLayoutBuilder::GetDetailFont())
-				]
-			]
+			SAssignNew(StructPicker, SInstancedStructPicker, StructProperty, PropUtils)
 		]
 		.IsEnabled(StructProperty->IsEditable());
 }
@@ -573,135 +448,4 @@ void FInstancedStructDetails::CustomizeChildren(TSharedRef<class IPropertyHandle
 	StructBuilder.AddCustomBuilder(DataDetails);
 }
 
-FText FInstancedStructDetails::GetDisplayValueString() const
-{
-	const UScriptStruct* CommonStruct = nullptr;
-	const FPropertyAccess::Result Result = UE::StructUtils::Private::GetCommonScriptStruct(StructProperty, CommonStruct);
-	
-	if (Result == FPropertyAccess::Success)
-	{
-		if (CommonStruct)
-		{
-			return CommonStruct->GetDisplayNameText();
-		}
-		return LOCTEXT("NullScriptStruct", "None");
-	}
-	if (Result == FPropertyAccess::MultipleValues)
-	{
-		return LOCTEXT("MultipleValues", "Multiple Values");
-	}
-	
-	return FText::GetEmpty();
-}
-
-FText FInstancedStructDetails::GetTooltipText() const
-{
-	const UScriptStruct* CommonStruct = nullptr;
-	const FPropertyAccess::Result Result = UE::StructUtils::Private::GetCommonScriptStruct(StructProperty, CommonStruct);
-	
-	if (CommonStruct && Result == FPropertyAccess::Success)
-	{
-		return CommonStruct->GetToolTipText();
-	}
-	
-	return GetDisplayValueString();
-}
-
-const FSlateBrush* FInstancedStructDetails::GetDisplayValueIcon() const
-{
-	const UScriptStruct* CommonStruct = nullptr;
-	if (UE::StructUtils::Private::GetCommonScriptStruct(StructProperty, CommonStruct) == FPropertyAccess::Success)
-	{
-		return FSlateIconFinder::FindIconBrushForClass(UScriptStruct::StaticClass());
-	}
-	
-	return nullptr;
-}
-
-TSharedRef<SWidget> FInstancedStructDetails::GenerateStructPicker()
-{
-	static const FName NAME_ExcludeBaseStruct = "ExcludeBaseStruct";
-	static const FName NAME_HideViewOptions = "HideViewOptions";
-	static const FName NAME_ShowTreeView = "ShowTreeView";
-
-	const bool bExcludeBaseStruct = StructProperty->HasMetaData(NAME_ExcludeBaseStruct);
-	const bool bAllowNone = !(StructProperty->GetMetaDataProperty()->PropertyFlags & CPF_NoClear);
-	const bool bHideViewOptions = StructProperty->HasMetaData(NAME_HideViewOptions);
-	const bool bShowTreeView = StructProperty->HasMetaData(NAME_ShowTreeView);
-
-	TSharedRef<FInstancedStructFilter> StructFilter = MakeShared<FInstancedStructFilter>();
-	StructFilter->BaseStruct = BaseScriptStruct;
-	StructFilter->bAllowUserDefinedStructs = BaseScriptStruct == nullptr; // Only allow user defined structs when BaseStruct is not set.
-	StructFilter->bAllowBaseStruct = !bExcludeBaseStruct;
-
-	if (GEditor && StructProperty)
-	{
-		FAssetReferenceFilterContext AssetReferenceFilterContext;
-		
-		TArray<UPackage*> OuterPackages;
-		StructProperty->GetOuterPackages(OuterPackages);
-		for (UPackage* OuterPackage : OuterPackages)
-		{
-			AssetReferenceFilterContext.ReferencingAssets.Add(FAssetData(OuterPackage));
-		}
-
-		StructFilter->AssetReferenceFilter = GEditor->MakeAssetReferenceFilter(AssetReferenceFilterContext);
-	}
-
-	const UScriptStruct* SelectedStruct = nullptr;
-	const FPropertyAccess::Result Result = UE::StructUtils::Private::GetCommonScriptStruct(StructProperty, SelectedStruct);
-	
-	FStructViewerInitializationOptions Options;
-	Options.bShowNoneOption = bAllowNone;
-	Options.StructFilter = StructFilter;
-	Options.NameTypeToDisplay = EStructViewerNameTypeToDisplay::DisplayName;
-	Options.DisplayMode = bShowTreeView ? EStructViewerDisplayMode::TreeView : EStructViewerDisplayMode::ListView;
-	Options.bAllowViewOptions = !bHideViewOptions;
-	Options.SelectedStruct = SelectedStruct;
-	
-	FOnStructPicked OnPicked(FOnStructPicked::CreateSP(this, &FInstancedStructDetails::OnStructPicked));
-
-	return SNew(SBox)
-		.WidthOverride(280)
-		[
-			SNew(SVerticalBox)
-			+ SVerticalBox::Slot()
-			.AutoHeight()
-			.MaxHeight(500)
-			[
-				FModuleManager::LoadModuleChecked<FStructViewerModule>("StructViewer").CreateStructViewer(Options, OnPicked)
-			]
-		];
-}
-
-void FInstancedStructDetails::OnStructPicked(const UScriptStruct* InStruct)
-{
-	if (StructProperty && StructProperty->IsValidHandle())
-	{
-		FScopedTransaction Transaction(LOCTEXT("OnStructPicked", "Set Struct"));
-
-		StructProperty->NotifyPreChange();
-
-		StructProperty->EnumerateRawData([InStruct](void* RawData, const int32 /*DataIndex*/, const int32 /*NumDatas*/)
-		{
-			if (FInstancedStruct* InstancedStruct = static_cast<FInstancedStruct*>(RawData))
-			{
-				InstancedStruct->InitializeAs(InStruct);
-			}
-			return true;
-		});
-
-		StructProperty->NotifyPostChange(EPropertyChangeType::ValueSet);
-		StructProperty->NotifyFinishedChangingProperties();
-
-		// Property tree will be invalid after changing the struct type, force update.
-		if (PropUtils.IsValid())
-		{
-			PropUtils->ForceRefresh();
-		}
-	}
-
-	ComboButton->SetIsOpen(false);
-}
-
 #undef LOCTEXT_NAMESPACE
diff --git a/Engine/Plugins/Experimental/StructUtils/Source/StructUtilsEditor/Private/SInstancedStructPicker.cpp b/Engine/Plugins/Experimental/StructUtils/Source/StructUtilsEditor/Private/SInstancedStructPicker.cpp
new file mode 100644
--- /dev/null
+++ b/Engine/Plugins/Experimental/StructUtils/Source/StructUtilsEditor/Private/SInstancedStructPicker.cpp
@@ -0,0 +1,290 @@
+// Copyright Epic Games, Inc. All Rights Reserved.
+
+#include "SInstancedStructPicker.h"
+#include "DetailLayoutBuilder.h"
+#include "Editor.h"
+#include "Engine/UserDefinedStruct.h"
+#include "InstancedStruct.h"
+#include "IPropertyUtilities.h"
+#include "Modules/ModuleManager.h"
+#include "PropertyHandle.h"
+#include "ScopedTransaction.h"
+#include "Styling/SlateIconFinder.h"
+#include "Widgets/Images/SImage.h"
+#include "Widgets/Input/SComboButton.h"
+#include "Widgets/Layout/SBox.h"
+
+#define LOCTEXT_NAMESPACE "StructUtilsEditor"
+
+namespace UE::StructUtils::Private
+{
+	FPropertyAccess::Result GetCommonScriptStruct(const TSharedPtr<IPropertyHandle>& StructProperty,
+		const UScriptStruct*& OutCommonStruct)
+	{
+		bool bHasResult = false;
+		bool bHasMultipleValues = false;
+
+		StructProperty->EnumerateConstRawData([&OutCommonStruct, &bHasResult, &bHasMultipleValues](const void* RawData, const int32 /*DataIndex*/, const int32 /*NumDatas*/)
+		{
+			if (const FInstancedStruct* InstancedStruct = static_cast<const FInstancedStruct*>(RawData))
+			{
+				const UScriptStruct* Struct = InstancedStruct->GetScriptStruct();
+
+				if (!bHasResult)
+				{
+					OutCommonStruct = Struct;
+				}
+				else if (OutCommonStruct != Struct)
+				{
+					bHasMultipleValues = true;
+				}
+
+				bHasResult = true;
+			}
+
+			return true;
+		});
+
+		if (bHasMultipleValues)
+		{
+			return FPropertyAccess::MultipleValues;
+		}
+
+		return bHasResult ? FPropertyAccess::Success : FPropertyAccess::Fail;
+	}
+} // UE::StructUtils::Private
+
+////////////////////////////////////
+
+bool FInstancedStructFilter::IsStructAllowed(const FStructViewerInitializationOptions& InInitOptions, const UScriptStruct* InStruct, TSharedRef<FStructViewerFilterFuncs> InFilterFuncs)
+{
+	if (InStruct->IsA<UUserDefinedStruct>())
+	{
+		return bAllowUserDefinedStructs;
+	}
+
+	if (InStruct == BaseStruct)
+	{
+		return bAllowBaseStruct;
+	}
+
+	static const FName NAME_HiddenMetaTag = "Hidden";
+	if (InStruct->HasMetaData(NAME_HiddenMetaTag))
+	{
+		return false;
+	}
+
+	if (AssetReferenceFilter.IsValid())
+	{
+		if (!AssetReferenceFilter->PassesFilter(FAssetData(InStruct)))
+		{
+			return false;
+		}
+	}
+
+	// Query the native struct to see if it has the correct parent type (if any)
+	const UScriptStruct* Struct = BaseStruct.Get();
+	return !Struct || InStruct->IsChildOf(Struct);
+}
+
+bool FInstancedStructFilter::IsUnloadedStructAllowed(const FStructViewerInitializationOptions& InInitOptions, const FSoftObjectPath& InStructPath, TSharedRef<FStructViewerFilterFuncs> InFilterFuncs)
+{
+	// User Defined Structs don't support inheritance, so only include them requested
+	return bAllowUserDefinedStructs;
+}
+
+////////////////////////////////////
+
+void SInstancedStructPicker::Construct(const FArguments& InArgs, TSharedPtr<IPropertyHandle> InStructProperty,
+	TSharedPtr<IPropertyUtilities> InPropertyUtils)
+{
+	static const FName NAME_BaseStruct = "BaseStruct";
+	static const FName NAME_StructTypeConst = "StructTypeConst";
+
+	OnStructPicked = InArgs._OnStructPicked;
+	StructProperty = MoveTemp(InStructProperty);
+	PropUtils = MoveTemp(InPropertyUtils);
+
+	if (!StructProperty.IsValid() || !PropUtils.IsValid())
+	{
+		return;
+	}
+
+	BaseScriptStruct = nullptr;
+	{
+		const FString& BaseStructName = StructProperty->GetMetaData(NAME_BaseStruct);
+		if (!BaseStructName.IsEmpty())
+		{
+			if (UScriptStruct* Struct = UClass::TryFindTypeSlow<UScriptStruct>(BaseStructName))
+			{
+				BaseScriptStruct = Struct;
+			}
+			else
+			{
+				BaseScriptStruct = LoadObject<UScriptStruct>(nullptr, *BaseStructName);
+			}
+		}
+	}
+
+	ChildSlot
+	[
+		SAssignNew(ComboButton, SComboButton)
+		.OnGetMenuContent(this, &SInstancedStructPicker::GenerateStructPicker)
+		.ContentPadding(0)
+		.IsEnabled(StructProperty->IsEditable() && !StructProperty->HasMetaData(NAME_StructTypeConst))
+		.ButtonContent()
+		[
+			SNew(SHorizontalBox)
+			+ SHorizontalBox::Slot()
+			.AutoWidth()
+			.VAlign(VAlign_Center)
+			.Padding(0.0f, 0.0f, 4.0f, 0.0f)
+			[
+				SNew(SImage)
+				.Image(this, &SInstancedStructPicker::GetDisplayValueIcon)
+			]
+			+ SHorizontalBox::Slot()
+			.VAlign(VAlign_Center)
+			[
+				SNew(STextBlock)
+				.Text(this, &SInstancedStructPicker::GetDisplayValueString)
+				.ToolTipText(this, &SInstancedStructPicker::GetTooltipText)
+				.Font(IDetailLayoutBuilder::GetDetailFont())
+			]
+		]
+	];
+}
+
+FText SInstancedStructPicker::GetDisplayValueString() const
+{
+	const UScriptStruct* CommonStruct = nullptr;
+	const FPropertyAccess::Result Result = UE::StructUtils::Private::GetCommonScriptStruct(StructProperty, CommonStruct);
+
+	if (Result == FPropertyAccess::Success)
+	{
+		if (CommonStruct)
+		{
+			return CommonStruct->GetDisplayNameText();
+		}
+		return LOCTEXT("NullScriptStruct", "None");
+	}
+	if (Result == FPropertyAccess::MultipleValues)
+	{
+		return LOCTEXT("MultipleValues", "Multiple Values");
+	}
+
+	return FText::GetEmpty();
+}
+
+FText SInstancedStructPicker::GetTooltipText() const
+{
+	const UScriptStruct* CommonStruct = nullptr;
+	const FPropertyAccess::Result Result = UE::StructUtils::Private::GetCommonScriptStruct(StructProperty, CommonStruct);
+
+	if (CommonStruct && Result == FPropertyAccess::Success)
+	{
+		return CommonStruct->GetToolTipText();
+	}
+
+	return GetDisplayValueString();
+}
+
+const FSlateBrush* SInstancedStructPicker::GetDisplayValueIcon() const
+{
+	const UScriptStruct* CommonStruct = nullptr;
+	if (UE::StructUtils::Private::GetCommonScriptStruct(StructProperty, CommonStruct) == FPropertyAccess::Success)
+	{
+		return FSlateIconFinder::FindIconBrushForClass(UScriptStruct::StaticClass());
+	}
+
+	return nullptr;
+}
+
+TSharedRef<SWidget> SInstancedStructPicker::GenerateStructPicker()
+{
+	static const FName NAME_ExcludeBaseStruct = "ExcludeBaseStruct";
+	static const FName NAME_HideViewOptions = "HideViewOptions";
+	static const FName NAME_ShowTreeView = "ShowTreeView";
+
+	const bool bExcludeBaseStruct = StructProperty->HasMetaData(NAME_ExcludeBaseStruct);
+	const bool bAllowNone = !(StructProperty->GetMetaDataProperty()->PropertyFlags & CPF_NoClear);
+	const bool bHideViewOptions = StructProperty->HasMetaData(NAME_HideViewOptions);
+	const bool bShowTreeView = StructProperty->HasMetaData(NAME_ShowTreeView);
+
+	TSharedRef<FInstancedStructFilter> StructFilter = MakeShared<FInstancedStructFilter>();
+	StructFilter->BaseStruct = BaseScriptStruct;
+	StructFilter->bAllowUserDefinedStructs = BaseScriptStruct.IsExplicitlyNull(); // Only allow user defined structs when BaseStruct is not set.
+	StructFilter->bAllowBaseStruct = !bExcludeBaseStruct;
+
+	if (GEditor && StructProperty)
+	{
+		FAssetReferenceFilterContext AssetReferenceFilterContext;
+
+		TArray<UPackage*> OuterPackages;
+		StructProperty->GetOuterPackages(OuterPackages);
+		for (UPackage* OuterPackage : OuterPackages)
+		{
+			AssetReferenceFilterContext.ReferencingAssets.Add(FAssetData(OuterPackage));
+		}
+
+		StructFilter->AssetReferenceFilter = GEditor->MakeAssetReferenceFilter(AssetReferenceFilterContext);
+	}
+
+	const UScriptStruct* SelectedStruct = nullptr;
+	UE::StructUtils::Private::GetCommonScriptStruct(StructProperty, SelectedStruct);
+
+	FStructViewerInitializationOptions Options;
+	Options.bShowNoneOption = bAllowNone;
+	Options.StructFilter = StructFilter;
+	Options.NameTypeToDisplay = EStructViewerNameTypeToDisplay::DisplayName;
+	Options.DisplayMode = bShowTreeView ? EStructViewerDisplayMode::TreeView : EStructViewerDisplayMode::ListView;
+	Options.bAllowViewOptions = !bHideViewOptions;
+	Options.SelectedStruct = SelectedStruct;
+
+	const FOnStructPicked OnPicked(FOnStructPicked::CreateSP(this, &SInstancedStructPicker::StructPicked));
+
+	return SNew(SBox)
+		.WidthOverride(280)
+		[
+			SNew(SVerticalBox)
+			+ SVerticalBox::Slot()
+			.AutoHeight()
+			.MaxHeight(500)
+			[
+				FModuleManager::LoadModuleChecked<FStructViewerModule>("StructViewer").CreateStructViewer(Options, OnPicked)
+			]
+		];
+}
+
+void SInstancedStructPicker::StructPicked(const UScriptStruct* InStruct)
+{
+	if (StructProperty && StructProperty->IsValidHandle())
+	{
+		FScopedTransaction Transaction(LOCTEXT("OnStructPicked", "Set Struct"));
+
+		StructProperty->NotifyPreChange();
+
+		StructProperty->EnumerateRawData([InStruct](void* RawData, const int32 /*DataIndex*/, const int32 /*NumDatas*/)
+		{
+			if (FInstancedStruct* InstancedStruct = static_cast<FInstancedStruct*>(RawData))
+			{
+				InstancedStruct->InitializeAs(InStruct);
+			}
+			return true;
+		});
+
+		StructProperty->NotifyPostChange(EPropertyChangeType::ValueSet);
+		StructProperty->NotifyFinishedChangingProperties();
+
+		// Property tree will be invalid after changing the struct type, force update.
+		if (PropUtils.IsValid())
+		{
+			PropUtils->ForceRefresh();
+		}
+	}
+
+	ComboButton->SetIsOpen(false);
+	OnStructPicked.ExecuteIfBound(InStruct);
+}
+
+#undef LOCTEXT_NAMESPACE
diff --git a/Engine/Plugins/Experimental/StructUtils/Source/StructUtilsEditor/Public/InstancedStructDetails.h b/Engine/Plugins/Experimental/StructUtils/Source/StructUtilsEditor/Public/InstancedStructDetails.h
--- a/Engine/Plugins/Experimental/StructUtils/Source/StructUtilsEditor/Public/InstancedStructDetails.h
+++ b/Engine/Plugins/Experimental/StructUtils/Source/StructUtilsEditor/Public/InstancedStructDetails.h
@@ -4,8 +4,6 @@
 
 #include "IPropertyTypeCustomization.h"
 #include "IDetailCustomNodeBuilder.h"
-#include "EditorUndoClient.h"
-#include "StructViewerFilter.h"
 
 class IAssetReferenceFilter;
 class IPropertyHandle;
@@ -14,7 +12,7 @@ class IDetailPropertyRow;
 class IPropertyHandle;
 class FStructOnScope;
 class SWidget;
-class SComboButton;
+class SInstancedStructPicker;
 struct FInstancedStruct;
 class FInstancedStructProvider;
 
@@ -37,20 +35,11 @@ private:
 
 	using FReplacementObjectMap = TMap<UObject*, UObject*>;
 	void OnObjectsReinstanced(const FReplacementObjectMap& ObjectMap);
-	
-	FText GetDisplayValueString() const;
-	FText GetTooltipText() const;
-	const FSlateBrush* GetDisplayValueIcon() const;
-	TSharedRef<SWidget> GenerateStructPicker();
-	void OnStructPicked(const UScriptStruct* InStruct);
 
 	/** Handle to the struct property being edited */
 	TSharedPtr<IPropertyHandle> StructProperty;
 
-	/** The base struct that we're allowing to be picked (controlled by the "BaseStruct" meta-data) */
-	UScriptStruct* BaseScriptStruct = nullptr;
-
-	TSharedPtr<SComboButton> ComboButton;
+	TSharedPtr<SInstancedStructPicker> StructPicker;
 	TSharedPtr<IPropertyUtilities> PropUtils;
 	
 	FDelegateHandle OnObjectsReinstancedHandle;
@@ -122,29 +111,6 @@ protected:
 	void OnStructLayoutChanges();
 };
 
-/**
- * Filter used by the instanced struct struct picker.
- */
- class STRUCTUTILSEDITOR_API FInstancedStructFilter : public IStructViewerFilter
-{
-public:
-	/** The base struct for the property that classes must be a child-of. */
-	const UScriptStruct* BaseStruct = nullptr;
-
-	// A flag controlling whether we allow UserDefinedStructs
-	bool bAllowUserDefinedStructs = false;
-
-	// A flag controlling whether we allow to select the BaseStruct
-	bool bAllowBaseStruct = true;
-
-	virtual bool IsStructAllowed(const FStructViewerInitializationOptions& InInitOptions, const UScriptStruct* InStruct, TSharedRef<FStructViewerFilterFuncs> InFilterFuncs) override;
-	virtual bool IsUnloadedStructAllowed(const FStructViewerInitializationOptions& InInitOptions, const FSoftObjectPath& InStructPath, TSharedRef<FStructViewerFilterFuncs> InFilterFuncs) override;
-
- 	// Optional filter to prevent selection of some structs e.g. ones in a plugin that is inaccessible from the object being edited
- 	TSharedPtr<IAssetReferenceFilter> AssetReferenceFilter;
-};
-
-
 #if UE_ENABLE_INCLUDE_ORDER_DEPRECATED_IN_5_2
 #include "CoreMinimal.h"
 #include "UnrealClient.h"
diff --git a/Engine/Plugins/Experimental/StructUtils/Source/StructUtilsEditor/Public/SInstancedStructPicker.h b/Engine/Plugins/Experimental/StructUtils/Source/StructUtilsEditor/Public/SInstancedStructPicker.h
new file mode 100644
--- /dev/null
+++ b/Engine/Plugins/Experimental/StructUtils/Source/StructUtilsEditor/Public/SInstancedStructPicker.h
@@ -0,0 +1,60 @@
+// Copyright Epic Games, Inc. All Rights Reserved.
+
+#pragma once
+
+#include "StructViewerFilter.h"
+#include "StructViewerModule.h"
+#include "Widgets/SCompoundWidget.h"
+
+class IPropertyHandle;
+class IPropertyUtilities;
+class IAssetReferenceFilter;
+class SComboButton;
+
+/**
+ * Filter used by the instanced struct struct picker.
+ */
+class STRUCTUTILSEDITOR_API FInstancedStructFilter : public IStructViewerFilter
+{
+public:
+	/** The base struct for the property that classes must be a child-of. */
+	TWeakObjectPtr<const UScriptStruct> BaseStruct = nullptr;
+
+	// A flag controlling whether we allow UserDefinedStructs
+	bool bAllowUserDefinedStructs = false;
+
+	// A flag controlling whether we allow to select the BaseStruct
+	bool bAllowBaseStruct = true;
+
+	virtual bool IsStructAllowed(const FStructViewerInitializationOptions& InInitOptions, const UScriptStruct* InStruct, TSharedRef<FStructViewerFilterFuncs> InFilterFuncs) override;
+	virtual bool IsUnloadedStructAllowed(const FStructViewerInitializationOptions& InInitOptions, const FSoftObjectPath& InStructPath, TSharedRef<FStructViewerFilterFuncs> InFilterFuncs) override;
+
+	// Optional filter to prevent selection of some structs e.g. ones in a plugin that is inaccessible from the object being edited
+	TSharedPtr<IAssetReferenceFilter> AssetReferenceFilter;
+};
+
+class STRUCTUTILSEDITOR_API SInstancedStructPicker : public SCompoundWidget
+{
+public:
+	SLATE_BEGIN_ARGS(SInstancedStructPicker) { }
+		SLATE_ARGUMENT(FOnStructPicked, OnStructPicked)
+	SLATE_END_ARGS()
+
+	void Construct(const FArguments& InArgs, TSharedPtr<IPropertyHandle> InStructProperty, TSharedPtr<IPropertyUtilities> InPropertyUtils);
+
+	FOnStructPicked OnStructPicked;
+
+private:
+	TSharedPtr<SComboButton> ComboButton;
+	TSharedPtr<IPropertyHandle> StructProperty;
+	TSharedPtr<IPropertyUtilities> PropUtils;
+
+	/** The base struct that we're allowing to be picked (controlled by the "BaseStruct" meta-data) */
+	TWeakObjectPtr<UScriptStruct> BaseScriptStruct = nullptr;
+
+	FText GetDisplayValueString() const;
+	FText GetTooltipText() const;
+	const FSlateBrush* GetDisplayValueIcon() const;
+	TSharedRef<SWidget> GenerateStructPicker();
+	void StructPicked(const UScriptStruct* InStruct);
+};
